#ttp 331133 - Editor: Source Control: Resize Bug in Setup Source Control Dialog
#proj UE4
#branch UE4
#summary The source control window is now sized 700px wide

diff --git a/Engine/Source/Developer/SourceControl/Private/SourceControlModule.cpp b/Engine/Source/Developer/SourceControl/Private/SourceControlModule.cpp
--- a/Engine/Source/Developer/SourceControl/Private/SourceControlModule.cpp
+++ b/Engine/Source/Developer/SourceControl/Private/SourceControlModule.cpp
@@ -123,9 +123,13 @@ void FSourceControlModule::ShowLoginDialog(const FSourceControlLoginClosed& InOn
 
 		// Setup the content for the created login window.
 		SourceControlLoginWindowPtr->SetContent(
-			SAssignNew(SourceControlLoginPtr, SSourceControlLogin)
-			.ParentWindow(SourceControlLoginWindowPtr)
-			.OnSourceControlLoginClosed(InOnSourceControlLoginClosed)
+			SNew(SBox)
+			.WidthOverride(700.0f)
+			[
+				SAssignNew(SourceControlLoginPtr, SSourceControlLogin)
+				.ParentWindow(SourceControlLoginWindowPtr)
+				.OnSourceControlLoginClosed(InOnSourceControlLoginClosed)
+			]
 			);
 
 		TSharedPtr<SWindow> RootWindow = FGlobalTabmanager::Get()->GetRootWindow();
