Switched SourceControlSettings code to use FPaths::GeneratedConfigDir() instead of trying to concat the path on its own.

diff --git a/Engine/Source/Developer/SourceControl/Private/SourceControlHelpers.cpp b/Engine/Source/Developer/SourceControl/Private/SourceControlHelpers.cpp
--- a/Engine/Source/Developer/SourceControl/Private/SourceControlHelpers.cpp
+++ b/Engine/Source/Developer/SourceControl/Private/SourceControlHelpers.cpp
@@ -26,7 +26,7 @@ const FString& GetSettingsIni()
 		static FString SourceControlSettingsIni;
 		if(SourceControlSettingsIni.Len() == 0)
 		{
-			const FString SourceControlSettingsDir = FPaths::GameSavedDir() + TEXT("Config/");
+			const FString SourceControlSettingsDir = FPaths::GeneratedConfigDir();
 			FConfigCacheIni::LoadGlobalIniFile(SourceControlSettingsIni, TEXT("SourceControlSettings"), NULL, NULL, false, false, true, *SourceControlSettingsDir);
 		}
 		return SourceControlSettingsIni;
