Changes for running 'clang -E'  to preprocess & then calling clang to compile the output in separate stages.
Necessary for using a distributed compiler on OS X. Updated for recent changes that caused more preprocess problems.

diff --git a/Engine/Source/Developer/FunctionalTesting/Private/Tests/FuncTestPerforming.cpp b/Engine/Source/Developer/FunctionalTesting/Private/Tests/FuncTestPerforming.cpp
--- a/Engine/Source/Developer/FunctionalTesting/Private/Tests/FuncTestPerforming.cpp
+++ b/Engine/Source/Developer/FunctionalTesting/Private/Tests/FuncTestPerforming.cpp
@@ -97,7 +97,7 @@ bool FStartFTestsOnMap::Update()
 /**
  * 
  */
-IMPLEMENT_COMPLEX_AUTOMATION_TEST( FFunctionalTestingMaps, "Maps.Functional Testing", EAutomationTestFlags::ATF_Game | EAutomationTestFlags::ATF_Editor )
+IMPLEMENT_COMPLEX_AUTOMATION_TEST( FFunctionalTestingMaps, "Maps.Functional Testing", (EAutomationTestFlags::ATF_Game | EAutomationTestFlags::ATF_Editor) )
 
 
 /** 
diff --git a/Engine/Source/Developer/SourceControl/Private/Tests/SourceControlTests.cpp b/Engine/Source/Developer/SourceControl/Private/Tests/SourceControlTests.cpp
--- a/Engine/Source/Developer/SourceControl/Private/Tests/SourceControlTests.cpp
+++ b/Engine/Source/Developer/SourceControl/Private/Tests/SourceControlTests.cpp
@@ -142,7 +142,7 @@ bool FSetProviderLatentCommand::Update()
 	return true;
 }
 
-IMPLEMENT_COMPLEX_AUTOMATION_TEST(FSetProviderTest, "Editor.Source Control.Set Provider", EAutomationTestFlags::ATF_Editor | EAutomationTestFlags::ATF_RequiresUser)
+IMPLEMENT_COMPLEX_AUTOMATION_TEST(FSetProviderTest, "Editor.Source Control.Set Provider", (EAutomationTestFlags::ATF_Editor | EAutomationTestFlags::ATF_RequiresUser))
 
 void FSetProviderTest::GetTests(TArray<FString>& OutBeautifiedNames, TArray<FString>& OutTestCommands) const
 {
@@ -173,7 +173,7 @@ bool FConnectLatentCommand::Update()
 	return AsyncHelper.IsDone();
 }
 
-IMPLEMENT_COMPLEX_AUTOMATION_TEST(FConnectTest, "Editor.Source Control.Connect", EAutomationTestFlags::ATF_Editor | EAutomationTestFlags::ATF_RequiresUser)
+IMPLEMENT_COMPLEX_AUTOMATION_TEST(FConnectTest, "Editor.Source Control.Connect", (EAutomationTestFlags::ATF_Editor | EAutomationTestFlags::ATF_RequiresUser))
 
 void FConnectTest::GetTests(TArray<FString>& OutBeautifiedNames, TArray<FString>& OutTestCommands) const
 {
@@ -270,7 +270,7 @@ bool FCheckOutLatentCommand::Update()
 	return AsyncHelper.IsDone();
 }
 
-IMPLEMENT_COMPLEX_AUTOMATION_TEST(FCheckOutTest, "Editor.Source Control.Check Out", EAutomationTestFlags::ATF_Editor | EAutomationTestFlags::ATF_RequiresUser)
+IMPLEMENT_COMPLEX_AUTOMATION_TEST(FCheckOutTest, "Editor.Source Control.Check Out", (EAutomationTestFlags::ATF_Editor | EAutomationTestFlags::ATF_RequiresUser))
 
 void FCheckOutTest::GetTests(TArray<FString>& OutBeautifiedNames, TArray<FString>& OutTestCommands) const
 {
@@ -421,7 +421,7 @@ bool FMarkForAddLatentCommand::Update()
 	return AsyncHelper.IsDone();
 }
 
-IMPLEMENT_COMPLEX_AUTOMATION_TEST(FMarkForAddTest, "Editor.Source Control.Mark For Add", EAutomationTestFlags::ATF_Editor | EAutomationTestFlags::ATF_RequiresUser)
+IMPLEMENT_COMPLEX_AUTOMATION_TEST(FMarkForAddTest, "Editor.Source Control.Mark For Add", (EAutomationTestFlags::ATF_Editor | EAutomationTestFlags::ATF_RequiresUser))
 
 void FMarkForAddTest::GetTests(TArray<FString>& OutBeautifiedNames, TArray<FString>& OutTestCommands) const
 {
@@ -494,7 +494,7 @@ bool FDeleteLatentCommand::Update()
 	return AsyncHelper.IsDone();
 }
 
-IMPLEMENT_COMPLEX_AUTOMATION_TEST(FDeleteTest, "Editor.Source Control.Delete", EAutomationTestFlags::ATF_Editor | EAutomationTestFlags::ATF_RequiresUser)
+IMPLEMENT_COMPLEX_AUTOMATION_TEST(FDeleteTest, "Editor.Source Control.Delete", (EAutomationTestFlags::ATF_Editor | EAutomationTestFlags::ATF_RequiresUser))
 
 void FDeleteTest::GetTests(TArray<FString>& OutBeautifiedNames, TArray<FString>& OutTestCommands) const
 {
@@ -595,7 +595,7 @@ bool FEditTextureLatentCommand::Update()
 	return true;
 }
 
-IMPLEMENT_COMPLEX_AUTOMATION_TEST(FCheckInTest, "Editor.Source Control.Check In", EAutomationTestFlags::ATF_Editor | EAutomationTestFlags::ATF_RequiresUser)
+IMPLEMENT_COMPLEX_AUTOMATION_TEST(FCheckInTest, "Editor.Source Control.Check In", (EAutomationTestFlags::ATF_Editor | EAutomationTestFlags::ATF_RequiresUser))
 
 void FCheckInTest::GetTests(TArray<FString>& OutBeautifiedNames, TArray<FString>& OutTestCommands) const
 {
@@ -665,7 +665,7 @@ bool FSyncLatentCommand::Update()
 	return AsyncHelper.IsDone();
 }
 
-IMPLEMENT_COMPLEX_AUTOMATION_TEST(FSyncTest, "Editor.Source Control.Sync", EAutomationTestFlags::ATF_Editor | EAutomationTestFlags::ATF_RequiresUser)
+IMPLEMENT_COMPLEX_AUTOMATION_TEST(FSyncTest, "Editor.Source Control.Sync", (EAutomationTestFlags::ATF_Editor | EAutomationTestFlags::ATF_RequiresUser))
 
 void FSyncTest::GetTests(TArray<FString>& OutBeautifiedNames, TArray<FString>& OutTestCommands) const
 {
@@ -688,7 +688,7 @@ bool FSyncTest::RunTest(const FString& Parameters)
 	return true;
 }
 
-IMPLEMENT_COMPLEX_AUTOMATION_TEST(FRevertTest, "Editor.Source Control.Revert", EAutomationTestFlags::ATF_Editor | EAutomationTestFlags::ATF_RequiresUser)
+IMPLEMENT_COMPLEX_AUTOMATION_TEST(FRevertTest, "Editor.Source Control.Revert", (EAutomationTestFlags::ATF_Editor | EAutomationTestFlags::ATF_RequiresUser))
 
 void FRevertTest::GetTests(TArray<FString>& OutBeautifiedNames, TArray<FString>& OutTestCommands) const
 {
@@ -783,7 +783,7 @@ bool FGetStateLatentCommand::Update()
 	return true;
 }
 
-IMPLEMENT_COMPLEX_AUTOMATION_TEST(FUpdateStatusTest, "Editor.Source Control.Update Status", EAutomationTestFlags::ATF_Editor | EAutomationTestFlags::ATF_RequiresUser)
+IMPLEMENT_COMPLEX_AUTOMATION_TEST(FUpdateStatusTest, "Editor.Source Control.Update Status", (EAutomationTestFlags::ATF_Editor | EAutomationTestFlags::ATF_RequiresUser))
 
 void FUpdateStatusTest::GetTests(TArray<FString>& OutBeautifiedNames, TArray<FString>& OutTestCommands) const
 {
@@ -867,7 +867,7 @@ bool FGetLabelLatentCommand::Update()
 	return true;
 }
 
-IMPLEMENT_COMPLEX_AUTOMATION_TEST(FGetLabelTest, "Editor.Source Control.Get Label", EAutomationTestFlags::ATF_Editor | EAutomationTestFlags::ATF_RequiresUser)
+IMPLEMENT_COMPLEX_AUTOMATION_TEST(FGetLabelTest, "Editor.Source Control.Get Label", (EAutomationTestFlags::ATF_Editor | EAutomationTestFlags::ATF_RequiresUser))
 
 void FGetLabelTest::GetTests(TArray<FString>& OutBeautifiedNames, TArray<FString>& OutTestCommands) const
 {
@@ -929,7 +929,7 @@ bool FGetRevisionLatentCommand::Update()
 	return true;
 }
 
-IMPLEMENT_COMPLEX_AUTOMATION_TEST(FGetRevisionTest, "Editor.Source Control.Get Revision", EAutomationTestFlags::ATF_Editor | EAutomationTestFlags::ATF_RequiresUser)
+IMPLEMENT_COMPLEX_AUTOMATION_TEST(FGetRevisionTest, "Editor.Source Control.Get Revision", (EAutomationTestFlags::ATF_Editor | EAutomationTestFlags::ATF_RequiresUser))
 
 void FGetRevisionTest::GetTests(TArray<FString>& OutBeautifiedNames, TArray<FString>& OutTestCommands) const
 {
diff --git a/Engine/Source/Editor/UnrealEd/Private/Tests/BlueprintAutomationTests.cpp b/Engine/Source/Editor/UnrealEd/Private/Tests/BlueprintAutomationTests.cpp
--- a/Engine/Source/Editor/UnrealEd/Private/Tests/BlueprintAutomationTests.cpp
+++ b/Engine/Source/Editor/UnrealEd/Private/Tests/BlueprintAutomationTests.cpp
@@ -19,7 +19,7 @@ DEFINE_LOG_CATEGORY_STATIC(LogBlueprintAutomationTests, Log, All);
 IMPLEMENT_COMPLEX_AUTOMATION_TEST( FBlueprintCompileOnLoadTest, "Blueprints.Compile-On-Load", EAutomationTestFlags::ATF_Editor )
 IMPLEMENT_COMPLEX_AUTOMATION_TEST( FBlueprintInstancesTest, "Blueprints.Instance Test", EAutomationTestFlags::ATF_Editor )
 IMPLEMENT_COMPLEX_AUTOMATION_TEST( FBlueprintReparentTest, "Blueprints.Reparent", EAutomationTestFlags::ATF_Editor )
-IMPLEMENT_COMPLEX_AUTOMATION_TEST( FBlueprintRenameAndCloneTest, "Blueprints.Rename And Clone", EAutomationTestFlags::ATF_Editor | EAutomationTestFlags::ATF_RequiresUser )
+IMPLEMENT_COMPLEX_AUTOMATION_TEST( FBlueprintRenameAndCloneTest, "Blueprints.Rename And Clone", (EAutomationTestFlags::ATF_Editor | EAutomationTestFlags::ATF_RequiresUser) )
 IMPLEMENT_COMPLEX_AUTOMATION_TEST( FCompileBlueprintsTest, "Blueprints.Compile Blueprints", EAutomationTestFlags::ATF_Editor )
 IMPLEMENT_COMPLEX_AUTOMATION_TEST( FCompileAnimBlueprintsTest, "Blueprints.Compile Anims", EAutomationTestFlags::ATF_Editor )
 
diff --git a/Engine/Source/Editor/UnrealEd/Private/Tests/EditorAutomationTests.cpp b/Engine/Source/Editor/UnrealEd/Private/Tests/EditorAutomationTests.cpp
--- a/Engine/Source/Editor/UnrealEd/Private/Tests/EditorAutomationTests.cpp
+++ b/Engine/Source/Editor/UnrealEd/Private/Tests/EditorAutomationTests.cpp
@@ -150,7 +150,7 @@ public:
  * in the Content Browser (does not allow for specific settings to be made per import factory). Cannot be run in a commandlet
  * as it executes code that routes through Slate UI.
  */
-IMPLEMENT_COMPLEX_AUTOMATION_TEST( FGenericImportAssetsAutomationTest, "Editor.Import", EAutomationTestFlags::ATF_Editor | EAutomationTestFlags::ATF_NonNullRHI )
+IMPLEMENT_COMPLEX_AUTOMATION_TEST( FGenericImportAssetsAutomationTest, "Editor.Import", (EAutomationTestFlags::ATF_Editor | EAutomationTestFlags::ATF_NonNullRHI) )
 
 /** 
  * Requests a enumeration of all sample assets to import
@@ -1491,7 +1491,7 @@ bool FLightPlacement::RunTest(const FString& Parameters)
  * Unit test to find all timelines in blueprints and list the events that can trigger them.
  * Timelines implicitly tick and are usually used for cosmetic events, so they can cause performance problems on dedicated servers.
  */
-IMPLEMENT_COMPLEX_AUTOMATION_TEST( FTraceAllTimelinesAutomationTest, "Performance Audits.Find Timelines On Server", EAutomationTestFlags::ATF_Editor | EAutomationTestFlags::ATF_RequiresUser)
+IMPLEMENT_COMPLEX_AUTOMATION_TEST( FTraceAllTimelinesAutomationTest, "Performance Audits.Find Timelines On Server", (EAutomationTestFlags::ATF_Editor | EAutomationTestFlags::ATF_RequiresUser))
 
 /** 
  * Requests an enumeration of all blueprints to be loaded
diff --git a/Engine/Source/Runtime/Engine/Private/Interpolation.cpp b/Engine/Source/Runtime/Engine/Private/Interpolation.cpp
--- a/Engine/Source/Runtime/Engine/Private/Interpolation.cpp
+++ b/Engine/Source/Runtime/Engine/Private/Interpolation.cpp
@@ -132,7 +132,7 @@ int32 TrackClass::SetKeyframeTime(int32 KeyIndex, float NewKeyTime, bool bUpdate
 		MoveKey.TimeVar = NewKeyTime; \
 		/* Find correct new position and insert. */ \
 		int32 i=0; \
-		for( i=0; i<KeyArray.Num() && KeyArray[i].TimeVar < NewKeyTime; i++); \
+		for( i=0; i<KeyArray.Num() && KeyArray[i].TimeVar < NewKeyTime; i++) {}; \
 		KeyArray.InsertZeroed(i); \
 		KeyArray[i] = MoveKey; \
 		return i; \
@@ -170,7 +170,7 @@ int32 TrackClass::DuplicateKeyframe(int32 KeyIndex, float NewKeyTime, UInterpTra
 	KeyType NewKey = KeyArray[KeyIndex]; \
 	NewKey.TimeVar = NewKeyTime; \
 	/* Find the correct index to insert this key. */ \
-	int32 i=0; for( i=0; i<DestTrack->KeyArray.Num() && DestTrack->KeyArray[i].TimeVar < NewKeyTime; i++); \
+	int32 i=0; for( i=0; i<DestTrack->KeyArray.Num() && DestTrack->KeyArray[i].TimeVar < NewKeyTime; i++) {}; \
 	DestTrack->KeyArray.InsertZeroed(i); \
 	DestTrack->KeyArray[i] = NewKey; \
 	return i; \
diff --git a/Engine/Source/Runtime/OpenGLDrv/Public/OpenGLUtil.h b/Engine/Source/Runtime/OpenGLDrv/Public/OpenGLUtil.h
--- a/Engine/Source/Runtime/OpenGLDrv/Public/OpenGLUtil.h
+++ b/Engine/Source/Runtime/OpenGLDrv/Public/OpenGLUtil.h
@@ -10,11 +10,11 @@
 #define ENABLE_OPENGL_FRAMEDUMP 0
 
 /** Set to 1 to enable the VERIFY_GL macros which call glGetError */
-#define ENABLE_VERIFY_GL (0 && !PLATFORM_MAC && UE_BUILD_DEBUG)
+#define ENABLE_VERIFY_GL (0 & !PLATFORM_MAC & UE_BUILD_DEBUG)
 #define ENABLE_VERIFY_GL_TRACE 0
 
 /** Set to 1 to verify that the the engine side uniform buffer layout matches the driver side of the GLSL shader*/
-#define ENABLE_UNIFORM_BUFFER_LAYOUT_VERIFICATION (UE_BUILD_DEBUG && (OPENGL_GL3 || OPENGL_GL4))
+#define ENABLE_UNIFORM_BUFFER_LAYOUT_VERIFICATION (UE_BUILD_DEBUG & (OPENGL_GL3 | OPENGL_GL4))
 
 /** Set to 1 to additinally dump uniform buffer layout at shader link time, this assumes ENABLE_UNIFORM_BUFFER_LAYOUT_VERIFICATION == 1 */
 #define ENABLE_UNIFORM_BUFFER_LAYOUT_DUMP 0
@@ -28,7 +28,7 @@
 
 #define ENABLE_OPENGL_DEBUG_GROUPS 1
 
-#define OPENGL_PERFORMANCE_DATA_INVALID (ENABLE_OPENGL_FRAMEDUMP || ENABLE_VERIFY_GL || ENABLE_UNIFORM_BUFFER_LAYOUT_VERIFICATION || DEBUG_GL_SHADERS)
+#define OPENGL_PERFORMANCE_DATA_INVALID (ENABLE_OPENGL_FRAMEDUMP | ENABLE_VERIFY_GL | ENABLE_UNIFORM_BUFFER_LAYOUT_VERIFICATION | DEBUG_GL_SHADERS)
 
 /**
 * Convert from ECubeFace to GLenum type
