Fixed source control module not reporting circular dependency on UnrealEd module

diff --git a/Engine/Source/Developer/SourceControl/Public/ISourceControlModule.h b/Engine/Source/Developer/SourceControl/Public/ISourceControlModule.h
--- a/Engine/Source/Developer/SourceControl/Public/ISourceControlModule.h
+++ b/Engine/Source/Developer/SourceControl/Public/ISourceControlModule.h
@@ -121,7 +121,7 @@ public:
 	 *
 	 * @return A reference to the source control module.
 	 */
-	SOURCECONTROL_API static ISourceControlModule& Get()
+	static inline ISourceControlModule& Get()
 	{
 		return FModuleManager::LoadModuleChecked<ISourceControlModule>("SourceControl");
 	}
diff --git a/Engine/Source/Developer/SourceControl/SourceControl.Build.cs b/Engine/Source/Developer/SourceControl/SourceControl.Build.cs
--- a/Engine/Source/Developer/SourceControl/SourceControl.Build.cs
+++ b/Engine/Source/Developer/SourceControl/SourceControl.Build.cs
@@ -41,6 +41,8 @@ public class SourceControl : ModuleRules
 					"AssetTools"
 				}
 			);
+	        
+			CircularlyReferencedDependentModules.Add("UnrealEd");
         }
 
 		if (UEBuildConfiguration.bBuildDeveloperTools)
