SourceControlTests.cpp and EditorAutomationTests.cpp were using an asset that has moved to the NotForLicensee folder.
EditorAutomationTest.cpp -- Just changed the reference to the cube mesh in the engine folder.
SourceControlTests.cpp -- updated the SourceControlTest location to the new one.

diff --git a/Engine/Source/Developer/SourceControl/Private/Tests/SourceControlTests.cpp b/Engine/Source/Developer/SourceControl/Private/Tests/SourceControlTests.cpp
--- a/Engine/Source/Developer/SourceControl/Private/Tests/SourceControlTests.cpp
+++ b/Engine/Source/Developer/SourceControl/Private/Tests/SourceControlTests.cpp
@@ -275,7 +275,7 @@ IMPLEMENT_COMPLEX_AUTOMATION_TEST(FCheckOutTest, "Editor.Source Control.Check Ou
 void FCheckOutTest::GetTests(TArray<FString>& OutBeautifiedNames, TArray<FString>& OutTestCommands) const
 {
 	GetProviders(OutBeautifiedNames, OutTestCommands);
-	AppendFilename(TEXT("/Engine/EditorAutomation/SourceControlTest"), OutBeautifiedNames, OutTestCommands);
+	AppendFilename(TEXT("/Engine/NotForLicensees/Automation/EditorAutomation/SourceControlTest"), OutBeautifiedNames, OutTestCommands);
 }
 
 bool FCheckOutTest::RunTest(const FString& Parameters)
@@ -426,7 +426,7 @@ IMPLEMENT_COMPLEX_AUTOMATION_TEST(FMarkForAddTest, "Editor.Source Control.Mark F
 void FMarkForAddTest::GetTests(TArray<FString>& OutBeautifiedNames, TArray<FString>& OutTestCommands) const
 {
 	GetProviders(OutBeautifiedNames, OutTestCommands);
-	AppendFilename(TEXT("/Engine/EditorAutomation/SourceControlTest"), OutBeautifiedNames, OutTestCommands);
+	AppendFilename(TEXT("/Engine/NotForLicensees/Automation/EditorAutomation/SourceControlTest"), OutBeautifiedNames, OutTestCommands);
 }
 
 bool FMarkForAddTest::RunTest(const FString& Parameters)
@@ -499,7 +499,7 @@ IMPLEMENT_COMPLEX_AUTOMATION_TEST(FDeleteTest, "Editor.Source Control.Delete", (
 void FDeleteTest::GetTests(TArray<FString>& OutBeautifiedNames, TArray<FString>& OutTestCommands) const
 {
 	GetProviders(OutBeautifiedNames, OutTestCommands);
-	AppendFilename(TEXT("/Engine/EditorAutomation/SourceControlTest"), OutBeautifiedNames, OutTestCommands);
+	AppendFilename(TEXT("/Engine/NotForLicensees/Automation/EditorAutomation/SourceControlTest"), OutBeautifiedNames, OutTestCommands);
 }
 
 bool FDeleteTest::RunTest(const FString& Parameters)
@@ -600,7 +600,7 @@ IMPLEMENT_COMPLEX_AUTOMATION_TEST(FCheckInTest, "Editor.Source Control.Check In"
 void FCheckInTest::GetTests(TArray<FString>& OutBeautifiedNames, TArray<FString>& OutTestCommands) const
 {
 	GetProviders(OutBeautifiedNames, OutTestCommands);
-	AppendFilename(TEXT("/Engine/EditorAutomation/SourceControlTest"), OutBeautifiedNames, OutTestCommands);
+	AppendFilename(TEXT("/Engine/NotForLicensees/Automation/EditorAutomation/SourceControlTest"), OutBeautifiedNames, OutTestCommands);
 }
 
 bool FCheckInTest::RunTest(const FString& Parameters)
@@ -670,7 +670,7 @@ IMPLEMENT_COMPLEX_AUTOMATION_TEST(FSyncTest, "Editor.Source Control.Sync", (EAut
 void FSyncTest::GetTests(TArray<FString>& OutBeautifiedNames, TArray<FString>& OutTestCommands) const
 {
 	GetProviders(OutBeautifiedNames, OutTestCommands);
-	AppendFilename(TEXT("/Engine/EditorAutomation/SourceControlTest"), OutBeautifiedNames, OutTestCommands);
+	AppendFilename(TEXT("/Engine/NotForLicensees/Automation/EditorAutomation/SourceControlTest"), OutBeautifiedNames, OutTestCommands);
 }
 
 bool FSyncTest::RunTest(const FString& Parameters)
@@ -693,7 +693,7 @@ IMPLEMENT_COMPLEX_AUTOMATION_TEST(FRevertTest, "Editor.Source Control.Revert", (
 void FRevertTest::GetTests(TArray<FString>& OutBeautifiedNames, TArray<FString>& OutTestCommands) const
 {
 	GetProviders(OutBeautifiedNames, OutTestCommands);
-	AppendFilename(TEXT("/Engine/EditorAutomation/SourceControlTest"), OutBeautifiedNames, OutTestCommands);
+	AppendFilename(TEXT("/Engine/NotForLicensees/Automation/EditorAutomation/SourceControlTest"), OutBeautifiedNames, OutTestCommands);
 }
 
 bool FRevertTest::RunTest(const FString& Parameters)
@@ -788,7 +788,7 @@ IMPLEMENT_COMPLEX_AUTOMATION_TEST(FUpdateStatusTest, "Editor.Source Control.Upda
 void FUpdateStatusTest::GetTests(TArray<FString>& OutBeautifiedNames, TArray<FString>& OutTestCommands) const
 {
 	GetProviders(OutBeautifiedNames, OutTestCommands);
-	AppendFilename(TEXT("/Engine/EditorAutomation/SourceControlTest"), OutBeautifiedNames, OutTestCommands);
+	AppendFilename(TEXT("/Engine/NotForLicensees/Automation/EditorAutomation/SourceControlTest"), OutBeautifiedNames, OutTestCommands);
 }
 
 bool FUpdateStatusTest::RunTest(const FString& Parameters)
@@ -983,7 +983,7 @@ IMPLEMENT_COMPLEX_AUTOMATION_TEST(FGetRevisionTest, "Editor.Source Control.Get R
 void FGetRevisionTest::GetTests(TArray<FString>& OutBeautifiedNames, TArray<FString>& OutTestCommands) const
 {
 	GetProviders(OutBeautifiedNames, OutTestCommands);
-	AppendFilename(TEXT("/Engine/EditorAutomation/SourceControlTest"), OutBeautifiedNames, OutTestCommands);
+	AppendFilename(TEXT("/Engine/NotForLicensees/Automation/EditorAutomation/SourceControlTest"), OutBeautifiedNames, OutTestCommands);
 }
 
 bool FGetRevisionTest::RunTest(const FString& Parameters)
diff --git a/Engine/Source/Editor/UnrealEd/Private/Tests/EditorAutomationTests.cpp b/Engine/Source/Editor/UnrealEd/Private/Tests/EditorAutomationTests.cpp
--- a/Engine/Source/Editor/UnrealEd/Private/Tests/EditorAutomationTests.cpp
+++ b/Engine/Source/Editor/UnrealEd/Private/Tests/EditorAutomationTests.cpp
@@ -1558,11 +1558,11 @@ bool FLightPlacement::RunTest(const FString& Parameters)
 
 	{
 		//Gather assets.
-		UObject* Astroid = (UStaticMesh*)StaticLoadObject(UStaticMesh::StaticClass(), NULL, TEXT("/Engine/Content/EditorAutomation/Astroid.Astroid"), NULL, LOAD_None, NULL);
-		//Add Astroid mesh to the world
-		AStaticMeshActor* StaticMesh = Cast<AStaticMeshActor>(FActorFactoryAssetProxy::AddActorForAsset(Astroid));
+		UObject* Cube = (UStaticMesh*)StaticLoadObject(UStaticMesh::StaticClass(), NULL, TEXT("/Engine/EngineMeshes/Cube.Cube"), NULL, LOAD_None, NULL);
+		//Add Cube mesh to the world
+		AStaticMeshActor* StaticMesh = Cast<AStaticMeshActor>(FActorFactoryAssetProxy::AddActorForAsset(Cube));
 		StaticMesh->TeleportTo(FVector(0.0f, 0.0f, 0.0f), FRotator(0, 0, 0));
-		StaticMesh->SetActorRelativeScale3D(FVector(1.0f, 1.0f, 1.0f));
+		StaticMesh->SetActorRelativeScale3D(FVector(3.0f, 3.0f, 1.75f));
 	}
 
 	//Create the point light and set it's mobility, brightness, and light color.
