Fixed crash closing the editor with SVN source control enabled

TTP# 342872 - EDITOR: Regression: CRASH: FSourceControlModule::ShutdownModule

The SVN source control provider module was being shutdown before the main source control module, leading to a crash.

The source control module now watches for its active provider module being unloaded, and will gracefully reset itself to the default (dummy) provider if that happens.

I also added these same checks to the source code access module, as it uses the same provider modules mechanism.

diff --git a/Engine/Source/Developer/SourceCodeAccess/Private/SourceCodeAccessModule.cpp b/Engine/Source/Developer/SourceCodeAccess/Private/SourceCodeAccessModule.cpp
--- a/Engine/Source/Developer/SourceCodeAccess/Private/SourceCodeAccessModule.cpp
+++ b/Engine/Source/Developer/SourceCodeAccess/Private/SourceCodeAccessModule.cpp
@@ -23,6 +23,7 @@ void FSourceCodeAccessModule::StartupModule()
 
 	// Register to check for source control features
 	IModularFeatures::Get().OnModularFeatureRegistered().AddRaw(this, &FSourceCodeAccessModule::HandleModularFeatureRegistered);
+	IModularFeatures::Get().OnModularFeatureUnregistered().AddRaw(this, &FSourceCodeAccessModule::HandleModularFeatureUnregistered);
 
 	// bind default accessor to editor
 	IModularFeatures::Get().RegisterModularFeature(SourceCodeAccessorFeatureName, &DefaultSourceCodeAccessor);
@@ -53,6 +54,7 @@ void FSourceCodeAccessModule::ShutdownModule()
 
 	// we don't care about modular features any more
 	IModularFeatures::Get().OnModularFeatureRegistered().RemoveAll(this);
+	IModularFeatures::Get().OnModularFeatureUnregistered().RemoveAll(this);
 }
 
 bool FSourceCodeAccessModule::CanAccessSourceCode() const
@@ -102,7 +104,7 @@ FOpenFileFailed& FSourceCodeAccessModule::OnOpenFileFailed()
 	return OpenFileFailedDelegate;
 }
 
-void FSourceCodeAccessModule::HandleModularFeatureRegistered(const FName& Type)
+void FSourceCodeAccessModule::HandleModularFeatureRegistered(const FName& Type, IModularFeature* ModularFeature)
 {
 	if(Type == SourceCodeAccessorFeatureName)
 	{
@@ -131,4 +133,12 @@ void FSourceCodeAccessModule::HandleModularFeatureRegistered(const FName& Type)
 	}
 }
 
+void FSourceCodeAccessModule::HandleModularFeatureUnregistered(const FName& Type, IModularFeature* ModularFeature)
+{
+	if(Type == SourceCodeAccessorFeatureName && CurrentSourceCodeAccessor == static_cast<ISourceCodeAccessor*>(ModularFeature))
+	{
+		CurrentSourceCodeAccessor = &DefaultSourceCodeAccessor;
+	}
+}
+
 #undef LOCTEXT_NAMESPACE 
\ No newline at end of file
diff --git a/Engine/Source/Developer/SourceCodeAccess/Private/SourceCodeAccessModule.h b/Engine/Source/Developer/SourceCodeAccess/Private/SourceCodeAccessModule.h
--- a/Engine/Source/Developer/SourceCodeAccess/Private/SourceCodeAccessModule.h
+++ b/Engine/Source/Developer/SourceCodeAccess/Private/SourceCodeAccessModule.h
@@ -27,7 +27,10 @@ public:
 
 private:
 	/** Handle when one of the modular features we are interested in is registered */
-	void HandleModularFeatureRegistered(const FName& Type);
+	void HandleModularFeatureRegistered(const FName& Type, IModularFeature* ModularFeature);
+
+	/** Handle when one of the modular features we are interested in is unregistered */
+	void HandleModularFeatureUnregistered(const FName& Type, IModularFeature* ModularFeature);
 
 private:
 	/** Event delegate fired when launching code accessor */
diff --git a/Engine/Source/Developer/SourceControl/Private/SourceControlModule.cpp b/Engine/Source/Developer/SourceControl/Private/SourceControlModule.cpp
--- a/Engine/Source/Developer/SourceControl/Private/SourceControlModule.cpp
+++ b/Engine/Source/Developer/SourceControl/Private/SourceControlModule.cpp
@@ -19,6 +19,8 @@ DEFINE_LOG_CATEGORY(LogSourceControl);
 
 #define LOCTEXT_NAMESPACE "SourceControl"
 
+static const FName SourceControlFeatureName("SourceControl");
+
 namespace SourceControlConstants
 {
 	/** The maximum number of file/directory status requests we should dispatch in a tick */
@@ -41,9 +43,10 @@ void FSourceControlModule::StartupModule()
 
 	// Register to check for source control features
 	IModularFeatures::Get().OnModularFeatureRegistered().AddRaw(this, &FSourceControlModule::HandleModularFeatureRegistered);
+	IModularFeatures::Get().OnModularFeatureUnregistered().AddRaw(this, &FSourceControlModule::HandleModularFeatureUnregistered);
 
 	// bind default provider to editor
-	IModularFeatures::Get().RegisterModularFeature( "SourceControl", &DefaultSourceControlProvider );
+	IModularFeatures::Get().RegisterModularFeature( SourceControlFeatureName, &DefaultSourceControlProvider );
 
 #if WITH_UNREAL_DEVELOPER_TOOLS
 	// create a message log for source control to use
@@ -67,10 +70,11 @@ void FSourceControlModule::ShutdownModule()
 #endif
 
 	// unbind default provider from editor
-	IModularFeatures::Get().UnregisterModularFeature( "SourceControl", &DefaultSourceControlProvider );
+	IModularFeatures::Get().UnregisterModularFeature( SourceControlFeatureName, &DefaultSourceControlProvider );
 
 	// we don't care about modular features any more
 	IModularFeatures::Get().OnModularFeatureRegistered().RemoveAll(this);
+	IModularFeatures::Get().OnModularFeatureUnregistered().RemoveAll(this);
 }
 
 void FSourceControlModule::SaveSettings()
@@ -189,11 +193,11 @@ void FSourceControlModule::InitializeSourceControlProviders()
 	// Look for valid SourceControl modules - they will register themselves as editor features
 	RefreshSourceControlProviders();
 
-	int32 SourceControlCount = IModularFeatures::Get().GetModularFeatureImplementationCount("SourceControl");
+	int32 SourceControlCount = IModularFeatures::Get().GetModularFeatureImplementationCount(SourceControlFeatureName);
 	if( SourceControlCount > 0 )
 	{
 		FString PreferredSourceControlProvider = SourceControlSettings.GetProvider();
-		TArray<ISourceControlProvider*> Providers = IModularFeatures::Get().GetModularFeatureImplementations<ISourceControlProvider>("SourceControl");
+		TArray<ISourceControlProvider*> Providers = IModularFeatures::Get().GetModularFeatureImplementations<ISourceControlProvider>(SourceControlFeatureName);
 		for(auto It(Providers.CreateIterator()); It; It++)
 		{ 
 			ISourceControlProvider* Provider = *It;
@@ -310,7 +314,7 @@ ISourceControlProvider& FSourceControlModule::GetProvider() const
 
 void FSourceControlModule::SetProvider( const FName& InName )
 {
-	TArray<ISourceControlProvider*> Providers = IModularFeatures::Get().GetModularFeatureImplementations<ISourceControlProvider>("SourceControl");
+	TArray<ISourceControlProvider*> Providers = IModularFeatures::Get().GetModularFeatureImplementations<ISourceControlProvider>(SourceControlFeatureName);
 	for(auto It(Providers.CreateIterator()); It; It++)
 	{
 		ISourceControlProvider* Provider = *It;
@@ -335,12 +339,12 @@ void FSourceControlModule::ClearCurrentSourceControlProvider()
 
 int32 FSourceControlModule::GetNumSourceControlProviders()
 {
-	return IModularFeatures::Get().GetModularFeatureImplementationCount("SourceControl");
+	return IModularFeatures::Get().GetModularFeatureImplementationCount(SourceControlFeatureName);
 }
 
 void FSourceControlModule::SetCurrentSourceControlProvider(int32 ProviderIndex)
 {
-	TArray<ISourceControlProvider*> Providers = IModularFeatures::Get().GetModularFeatureImplementations<ISourceControlProvider>("SourceControl");
+	TArray<ISourceControlProvider*> Providers = IModularFeatures::Get().GetModularFeatureImplementations<ISourceControlProvider>(SourceControlFeatureName);
 	check(Providers.IsValidIndex(ProviderIndex));
 	SetCurrentSourceControlProvider(*Providers[ProviderIndex]);
 }
@@ -365,7 +369,7 @@ void FSourceControlModule::SetCurrentSourceControlProvider(ISourceControlProvide
 
 FName FSourceControlModule::GetSourceControlProviderName(int32 ProviderIndex)
 {
-	TArray<ISourceControlProvider*> Providers = IModularFeatures::Get().GetModularFeatureImplementations<ISourceControlProvider>("SourceControl");
+	TArray<ISourceControlProvider*> Providers = IModularFeatures::Get().GetModularFeatureImplementations<ISourceControlProvider>(SourceControlFeatureName);
 	check(Providers.IsValidIndex(ProviderIndex));
 	return Providers[ProviderIndex]->GetName();
 }
@@ -375,14 +379,22 @@ TSharedPtr<class SSourceControlLogin> FSourceControlModule::GetLoginWidget() con
 	return SourceControlLoginPtr;
 }
 
-void FSourceControlModule::HandleModularFeatureRegistered(const FName& Type)
+void FSourceControlModule::HandleModularFeatureRegistered(const FName& Type, IModularFeature* ModularFeature)
 {
-	if(Type == TEXT("SourceControl"))
+	if(Type == SourceControlFeatureName)
 	{
 		InitializeSourceControlProviders();
 	}
 }
 
+void FSourceControlModule::HandleModularFeatureUnregistered(const FName& Type, IModularFeature* ModularFeature)
+{
+	if(Type == SourceControlFeatureName && CurrentSourceControlProvider == static_cast<ISourceControlProvider*>(ModularFeature))
+	{
+		ClearCurrentSourceControlProvider();
+	}
+}
+
 bool FSourceControlModule::GetUseGlobalSettings() const
 {
 	return SourceControlSettings.GetUseGlobalSettings();
diff --git a/Engine/Source/Developer/SourceControl/Private/SourceControlModule.h b/Engine/Source/Developer/SourceControl/Private/SourceControlModule.h
--- a/Engine/Source/Developer/SourceControl/Private/SourceControlModule.h
+++ b/Engine/Source/Developer/SourceControl/Private/SourceControlModule.h
@@ -87,7 +87,10 @@ private:
 	void OnSourceControlDialogClosed(const TSharedRef<class SWindow>& InWindow);
 
 	/** Delegate handling when source control features are registered */
-	void HandleModularFeatureRegistered(const FName& Type);
+	void HandleModularFeatureRegistered(const FName& Type, IModularFeature* ModularFeature);
+
+	/** Delegate handling when source control features are unregistered */
+	void HandleModularFeatureUnregistered(const FName& Type, IModularFeature* ModularFeature);
 
 private:
 	/** The settings object */
diff --git a/Engine/Source/Runtime/Core/Private/Features/ModularFeatures.cpp b/Engine/Source/Runtime/Core/Private/Features/ModularFeatures.cpp
--- a/Engine/Source/Runtime/Core/Private/Features/ModularFeatures.cpp
+++ b/Engine/Source/Runtime/Core/Private/Features/ModularFeatures.cpp
@@ -30,14 +30,14 @@ IModularFeature* FModularFeatures::GetModularFeatureImplementation( const FName
 void FModularFeatures::RegisterModularFeature( const FName Type, IModularFeature* ModularFeature )
 {
 	ModularFeaturesMap.AddUnique( Type, ModularFeature );
-	ModularFeatureRegisteredEvent.Broadcast( Type );
+	ModularFeatureRegisteredEvent.Broadcast( Type, ModularFeature );
 }
 
 
 void FModularFeatures::UnregisterModularFeature( const FName Type, IModularFeature* ModularFeature )
 {
 	ModularFeaturesMap.RemoveSingle( Type, ModularFeature );
-	ModularFeatureUnregisteredEvent.Broadcast( Type );
+	ModularFeatureUnregisteredEvent.Broadcast( Type, ModularFeature );
 }
 
 IModularFeatures::FOnModularFeatureRegistered& FModularFeatures::OnModularFeatureRegistered()
diff --git a/Engine/Source/Runtime/Core/Public/Features/IModularFeatures.h b/Engine/Source/Runtime/Core/Public/Features/IModularFeatures.h
--- a/Engine/Source/Runtime/Core/Public/Features/IModularFeatures.h
+++ b/Engine/Source/Runtime/Core/Public/Features/IModularFeatures.h
@@ -115,7 +115,7 @@ public:
 	 *
 	 * @param	Type	The name of the modular feature type being registered.
 	 */
-	DECLARE_EVENT_OneParam(IModularFeatures, FOnModularFeatureRegistered, const FName& /** Type */);
+	DECLARE_EVENT_TwoParams(IModularFeatures, FOnModularFeatureRegistered, const FName& /** Type */, class IModularFeature* /*ModularFeature*/);
 	virtual FOnModularFeatureRegistered& OnModularFeatureRegistered() = 0;
 
 	/** 
@@ -123,7 +123,7 @@ public:
 	 *
 	 * @param	Type	The name of the modular feature type being unregistered.
 	 */
-	DECLARE_EVENT_OneParam(IModularFeatures, FOnModularFeatureUnregistered, const FName& /** Type */);
+	DECLARE_EVENT_TwoParams(IModularFeatures, FOnModularFeatureUnregistered, const FName& /** Type */, class IModularFeature* /*ModularFeature*/);
 	virtual FOnModularFeatureUnregistered& OnModularFeatureUnregistered() = 0;
 };
 
