Make SourceControl module work (better) w/o Slate.

diff --git a/Engine/Source/Developer/SourceControl/Private/SourceControlModule.cpp b/Engine/Source/Developer/SourceControl/Private/SourceControlModule.cpp
--- a/Engine/Source/Developer/SourceControl/Private/SourceControlModule.cpp
+++ b/Engine/Source/Developer/SourceControl/Private/SourceControlModule.cpp
@@ -15,8 +15,6 @@
 
 DEFINE_LOG_CATEGORY(LogSourceControl);
 
-#if SOURCE_CONTROL_WITH_SLATE
-
 #define LOCTEXT_NAMESPACE "SourceControl"
 
 static const FName SourceControlFeatureName("SourceControl");
@@ -84,6 +82,7 @@ void FSourceControlModule::SaveSettings()
 
 void FSourceControlModule::ShowLoginDialog(const FSourceControlLoginClosed& InOnSourceControlLoginClosed, ELoginWindowMode::Type InLoginWindowMode, EOnLoginWindowStartup::Type InOnLoginWindowStartup)
 {
+#if SOURCE_CONTROL_WITH_SLATE
 	// Get Active Provider Name
 	ActiveProviderName = GetProvider().GetName().ToString();
 
@@ -153,11 +152,19 @@ void FSourceControlModule::ShowLoginDialog(const FSourceControlLoginClosed& InOn
 			FSlateApplication::Get().AddWindow(SourceControlLoginWindowPtr.ToSharedRef());
 		}
 	}
+#else
+	STUBBED("FSourceControlModule::ShowLoginDialog - no Slate");
+#endif // SOURCE_CONTROL_WITH_SLATE
 }
 
 TSharedPtr<class SWidget> FSourceControlModule::CreateStatusWidget() const
 {
+#if SOURCE_CONTROL_WITH_SLATE
 	return SNew(SSourceControlStatus);
+#else
+	STUBBED("FSourceControlModule::CreateStatusWidget - no Slate");
+	return nullptr;
+#endif // SOURCE_CONTROL_WITH_SLATE
 }
 
 void FSourceControlModule::OnSourceControlDialogClosed(const TSharedRef<SWindow>& InWindow)
@@ -411,5 +418,3 @@ void FSourceControlModule::SetUseGlobalSettings(bool bIsUseGlobalSettings)
 IMPLEMENT_MODULE( FSourceControlModule, SourceControl );
 
 #undef LOCTEXT_NAMESPACE
-
-#endif // SOURCE_CONTROL_WITH_SLATE
