Fix from Thomas Sarkanen for CRITICAL: EDITOR: Mac: CRASH: Creating a new collection while connented to Source Control (TTP #345087) - make sure that the scoped source control progress window is created as a modal window.

diff --git a/Engine/Source/Developer/SourceControl/Private/ScopedSourceControlProgress.cpp b/Engine/Source/Developer/SourceControl/Private/ScopedSourceControlProgress.cpp
--- a/Engine/Source/Developer/SourceControl/Private/ScopedSourceControlProgress.cpp
+++ b/Engine/Source/Developer/SourceControl/Private/ScopedSourceControlProgress.cpp
@@ -192,14 +192,7 @@ FScopedSourceControlProgress::FScopedSourceControlProgress(const FText& InText,
 		);
 			
 		TSharedPtr<SWindow> RootWindow = FGlobalTabmanager::Get()->GetRootWindow();
-		if(RootWindow.IsValid())
-		{
-			FSlateApplication::Get().AddModalWindow(Window, RootWindow, true);
-		}
-		else
-		{
-			FSlateApplication::Get().AddWindow(Window);
-		}
+		FSlateApplication::Get().AddModalWindow(Window, RootWindow, true);
 
 		Window->ShowWindow();	
 	
