SCC: Source control login UI polish

diff --git a/Engine/Source/Developer/SourceControl/Private/SSourceControlLogin.cpp b/Engine/Source/Developer/SourceControl/Private/SSourceControlLogin.cpp
--- a/Engine/Source/Developer/SourceControl/Private/SSourceControlLogin.cpp
+++ b/Engine/Source/Developer/SourceControl/Private/SSourceControlLogin.cpp
@@ -70,7 +70,7 @@ void SSourceControlLogin::Construct(const FArguments& InArgs)
 		SNew(SBorder)
 		.HAlign(HAlign_Fill)
 		.BorderImage( FEditorStyle::GetBrush("ChildWindow.Background") )
-		.Padding(0.0f)
+		.Padding(4.0f)
 		[
 			SNew(SVerticalBox)
 			+SVerticalBox::Slot()
@@ -87,51 +87,57 @@ void SSourceControlLogin::Construct(const FArguments& InArgs)
 			]
 			+SVerticalBox::Slot()
 			.FillHeight(1.0f)
-			.Padding(8.0f, 4.0f)
+			.Padding(0.0f, 6.0f, 0.0f, 0.0f)
 			[
-				SNew(SVerticalBox)
-				+SVerticalBox::Slot()
-				.AutoHeight()
-				.Padding(0.0f)
-				[
-					SNew(SBox)
-					.WidthOverride(500)
-					[
-						SNew(SSourceControlPicker)
-						.IsEnabled( this, &SSourceControlLogin::AreControlsEnabled )
-					]
-				]
-				+SVerticalBox::Slot()
-				.AutoHeight()
-				.Padding(0.0f, 0.0f, 0.0f, 4.0f)
-				[
-					SAssignNew(SettingsBorder, SBorder)
-					.BorderImage(FEditorStyle::GetBrush("NoBorder"))
-					.Visibility(this, &SSourceControlLogin::GetSettingsVisibility)
-					.IsEnabled(this, &SSourceControlLogin::AreControlsEnabled)
-					.Padding(0.0f)
-					[
-						SourceControlModule.GetProvider().MakeSettingsWidget()
-					]
-				]
-				+SVerticalBox::Slot()
-				.AutoHeight()
-				.Padding(0.0f, 0.0f, 0.0f, 4.0f)
-				[
-					SNew(SBorder)
-					.BorderImage(FEditorStyle::GetBrush("DetailsView.CategoryBottom"))
-					.Visibility(this, &SSourceControlLogin::GetDisabledTextVisibility)
-					.Padding(FMargin(4.0f, 12.0f))
+				SNew(SBorder)
+					.BorderImage(FEditorStyle::GetBrush("ToolPanel.GroupBorder"))
+					.Padding(4.0f)
 					[
-						SNew(STextBlock)
-						.WrapTextAt(500.0f)
-						.Text(LOCTEXT("SourceControlDisabledText", "Source control is currently disabled.\nTo enable, select a provider from the drop-down box above and enter your credentials.\nYou can re-enable source control by clicking on the icon in the top-right corner of the editor."))
+						SNew(SVerticalBox)
+						+SVerticalBox::Slot()
+						.AutoHeight()
+						.Padding(0.0f)
+						[
+							SNew(SBox)
+							.WidthOverride(500)
+							[
+								SNew(SSourceControlPicker)
+								.IsEnabled( this, &SSourceControlLogin::AreControlsEnabled )
+							]
+						]
+						+SVerticalBox::Slot()
+						.AutoHeight()
+						.Padding(0.0f, 0.0f, 0.0f, 4.0f)
+						[
+							SAssignNew(SettingsBorder, SBorder)
+							.BorderImage(FEditorStyle::GetBrush("NoBorder"))
+							.Visibility(this, &SSourceControlLogin::GetSettingsVisibility)
+							.IsEnabled(this, &SSourceControlLogin::AreControlsEnabled)
+							.Padding(0.0f)
+							[
+								SourceControlModule.GetProvider().MakeSettingsWidget()
+							]
+						]
+						+SVerticalBox::Slot()
+						.AutoHeight()
+						.Padding(0.0f, 0.0f, 0.0f, 4.0f)
+						[
+							SNew(SBorder)
+							.BorderImage(FEditorStyle::GetBrush("DetailsView.CategoryBottom"))
+							.Visibility(this, &SSourceControlLogin::GetDisabledTextVisibility)
+							.Padding(FMargin(4.0f, 12.0f))
+							[
+								SNew(STextBlock)
+								.WrapTextAt(500.0f)
+								.Text(LOCTEXT("SourceControlDisabledText", "Source control is currently disabled.\n\nTo enable, select a provider from the drop-down box above and enter your credentials.\nYou can re-enable source control by clicking on the icon in the top-right corner of the editor."))
+							]
+						]
 					]
-				]
 			]
 #if WITH_UNREAL_DEVELOPER_TOOLS
 			+SVerticalBox::Slot()
 			.AutoHeight()
+			.Padding(0.0f, 6.0f, 0.0f, 0.0f)
 			[
 				SNew(SExpandableArea)
 				.AreaTitleFont(FEditorStyle::GetFontStyle(TEXT("SourceControl.LoginWindow.Font")))
@@ -145,17 +151,14 @@ void SSourceControlLogin::Construct(const FArguments& InArgs)
 					.HeightOverride(250)
 					.WidthOverride(400)
 					[
-						SNew(SBorder)
-						[
-							MessageLogModule.CreateLogListingWidget(MessageLogListing)
-						]
+						MessageLogModule.CreateLogListingWidget(MessageLogListing)
 					]
 				]
 			]
 #endif
 			+SVerticalBox::Slot()
 			.AutoHeight()
-			.Padding(8.0f)
+			.Padding(8.0f, 16.0f, 8.0f, 8.0f)
 			[
 				SNew(SHorizontalBox)
 				+SHorizontalBox::Slot()
