FSourceControlModule::ShowLoginDialog will now properly show the window as modal if needed when RootWindow is NULL

diff --git a/Engine/Source/Developer/SourceControl/Private/SourceControlModule.cpp b/Engine/Source/Developer/SourceControl/Private/SourceControlModule.cpp
--- a/Engine/Source/Developer/SourceControl/Private/SourceControlModule.cpp
+++ b/Engine/Source/Developer/SourceControl/Private/SourceControlModule.cpp
@@ -149,7 +149,14 @@ void FSourceControlModule::ShowLoginDialog(const FSourceControlLoginClosed& InOn
 		}
 		else
 		{
-			FSlateApplication::Get().AddWindow(SourceControlLoginWindowPtr.ToSharedRef());
+			if(InLoginWindowMode == ELoginWindowMode::Modal)
+			{
+				FSlateApplication::Get().AddModalWindow(SourceControlLoginWindowPtr.ToSharedRef(), RootWindow);
+			}
+			else
+			{
+				FSlateApplication::Get().AddWindow(SourceControlLoginWindowPtr.ToSharedRef());
+			}
 		}
 	}
 #else
