Correctly sorted default (None) source control provider to the top of list

https://jira.ol.epicgames.net/browse/UE-7754

diff --git a/Engine/Source/Developer/SourceControl/Private/SSourceControlPicker.cpp b/Engine/Source/Developer/SourceControl/Private/SSourceControlPicker.cpp
--- a/Engine/Source/Developer/SourceControl/Private/SSourceControlPicker.cpp
+++ b/Engine/Source/Developer/SourceControl/Private/SSourceControlPicker.cpp
@@ -72,13 +72,14 @@ TSharedRef<SWidget> SSourceControlPicker::OnGetMenuContent() const
 	for(int ProviderIndex = 0; ProviderIndex < NumProviders; ++ProviderIndex)
 	{
 		const FName ProviderName = SourceControlModule.GetSourceControlProviderName(ProviderIndex);
-		SortedProviderNames.Add(TPairInitializer<FName, int32>(ProviderName, ProviderIndex));
+		int32 ProviderSortKey = ProviderName == FName("None") ? -1 * ProviderIndex : ProviderIndex;
+		SortedProviderNames.Add(TPairInitializer<FName, int32>(ProviderName, ProviderSortKey));
 	}
 
-	// Sort based on the provider name
+	// Sort based on the provider index
 	SortedProviderNames.Sort([](const TPair<FName, int32>& One, const TPair<FName, int32>& Two)
 	{
-		return One.Key < Two.Key;
+		return One.Value < Two.Value;
 	});
 
 	for(auto SortedProviderNameIter = SortedProviderNames.CreateConstIterator(); SortedProviderNameIter; ++SortedProviderNameIter)
@@ -93,7 +94,7 @@ TSharedRef<SWidget> SSourceControlPicker::OnGetMenuContent() const
 			FText::Format(LOCTEXT("SourceControlProvider_Tooltip", "Use {ProviderName} as source control provider"), Arguments),
 			FSlateIcon(),
 			FUIAction(
-				FExecuteAction::CreateSP( this, &SSourceControlPicker::ChangeSourceControlProvider, ProviderIndex ),
+				FExecuteAction::CreateSP( this, &SSourceControlPicker::ChangeSourceControlProvider, FMath::Abs(ProviderIndex) ),
 				FCanExecuteAction()
 				)
 			);
