Fixed behavior on FString::ParseIntoArray (muliple delimiters overload) functionality to support optionally culling empty strings. Greatly simplified implementation logic. Output parameter now properly named and taken by reference.

diff --git a/Engine/Plugins/Developer/GitSourceControl/Source/GitSourceControl/Private/GitSourceControlUtils.cpp b/Engine/Plugins/Developer/GitSourceControl/Source/GitSourceControl/Private/GitSourceControlUtils.cpp
--- a/Engine/Plugins/Developer/GitSourceControl/Source/GitSourceControl/Private/GitSourceControlUtils.cpp
+++ b/Engine/Plugins/Developer/GitSourceControl/Source/GitSourceControl/Private/GitSourceControlUtils.cpp
@@ -104,8 +104,8 @@ static bool RunCommandInternal(const FString& InCommand, const FString& InPathTo
 	FString Errors;
 
 	bResult = RunCommandInternalRaw(InCommand, InPathToGitBinary, InRepositoryRoot, InParameters, InFiles, Results, Errors);
-	Results.ParseIntoArray(&OutResults, TEXT("\n"), true);
-	Errors.ParseIntoArray(&OutErrorMessages, TEXT("\n"), true);
+	Results.ParseIntoArray(OutResults, TEXT("\n"), true);
+	Errors.ParseIntoArray(OutErrorMessages, TEXT("\n"), true);
 
 	return bResult;
 }
diff --git a/Engine/Plugins/Developer/PerforceSourceControl/Source/PerforceSourceControl/Private/PerforceConnection.cpp b/Engine/Plugins/Developer/PerforceSourceControl/Source/PerforceSourceControl/Private/PerforceConnection.cpp
--- a/Engine/Plugins/Developer/PerforceSourceControl/Source/PerforceSourceControl/Private/PerforceConnection.cpp
+++ b/Engine/Plugins/Developer/PerforceSourceControl/Source/PerforceSourceControl/Private/PerforceConnection.cpp
@@ -92,7 +92,7 @@ public:
 		OutputDesc += TEXT("Description:\n");
 		{
 			TArray<FString> DescLines;
-			Description.ToString().ParseIntoArray(&DescLines, TEXT("\n"), false);
+			Description.ToString().ParseIntoArray(DescLines, TEXT("\n"), false);
 			for (const FString& DescLine : DescLines)
 			{
 				OutputDesc += TEXT("\t");
diff --git a/Engine/Plugins/Developer/VisualStudioSourceCodeAccess/Source/VisualStudioSourceCodeAccess/Private/VisualStudioSourceCodeAccessor.cpp b/Engine/Plugins/Developer/VisualStudioSourceCodeAccess/Source/VisualStudioSourceCodeAccess/Private/VisualStudioSourceCodeAccessor.cpp
--- a/Engine/Plugins/Developer/VisualStudioSourceCodeAccess/Source/VisualStudioSourceCodeAccess/Private/VisualStudioSourceCodeAccessor.cpp
+++ b/Engine/Plugins/Developer/VisualStudioSourceCodeAccess/Source/VisualStudioSourceCodeAccess/Private/VisualStudioSourceCodeAccessor.cpp
@@ -875,7 +875,7 @@ bool FVisualStudioSourceCodeAccessor::AddSourceFiles(const TArray<FString>& Abso
 						{
 							const FString ProjectRelativeSourceFilePath = SourceFile.Mid(ModuleNewSourceFiles.ModuleNameAndPath.ModulePath.Len());
 							TArray<FString> SourceFileParts;
-							ProjectRelativeSourceFilePath.ParseIntoArray(&SourceFileParts, TEXT("/"), true);
+							ProjectRelativeSourceFilePath.ParseIntoArray(SourceFileParts, TEXT("/"), true);
 					
 							if (SourceFileParts.Num() == 0)
 							{
diff --git a/Engine/Plugins/Editor/PluginsEditor/Source/PluginsEditor/Private/SPluginCategories.cpp b/Engine/Plugins/Editor/PluginsEditor/Source/PluginsEditor/Private/SPluginCategories.cpp
--- a/Engine/Plugins/Editor/PluginsEditor/Source/PluginsEditor/Private/SPluginCategories.cpp
+++ b/Engine/Plugins/Editor/PluginsEditor/Source/PluginsEditor/Private/SPluginCategories.cpp
@@ -97,7 +97,7 @@ void SPluginCategories::RebuildAndFilterCategoryTree()
 				// We're expecting the category string to be in the "A.B.C" format.  We'll split up the string here and form
 				// a proper hierarchy in the UI
 				TArray< FString > SplitCategories;
-				CategoryPath.ParseIntoArray( &SplitCategories, TEXT( "." ), true /* bCullEmpty */ );
+				CategoryPath.ParseIntoArray( SplitCategories, TEXT( "." ), true /* bCullEmpty */ );
 
 				// Make sure all of the categories exist
 				for( auto SplitCategoryIt( SplitCategories.CreateConstIterator() ); SplitCategoryIt; ++SplitCategoryIt )
diff --git a/Engine/Plugins/Experimental/LiveEditor/Source/LiveEditor/Private/LiveEditorManager.cpp b/Engine/Plugins/Experimental/LiveEditor/Source/LiveEditor/Private/LiveEditorManager.cpp
--- a/Engine/Plugins/Experimental/LiveEditor/Source/LiveEditor/Private/LiveEditorManager.cpp
+++ b/Engine/Plugins/Experimental/LiveEditor/Source/LiveEditor/Private/LiveEditorManager.cpp
@@ -286,7 +286,7 @@ bool FLiveEditorManager::ConnectToRemoteHost( FString IPAddress )
 	IPAddress.Replace( TEXT(" "), TEXT("") );
 
 	TArray<FString> Parts;
-	IPAddress.ParseIntoArray( &Parts, TEXT("."), true );
+	IPAddress.ParseIntoArray( Parts, TEXT("."), true );
 	if ( Parts.Num() != 4 )
 		return false;
 
diff --git a/Engine/Plugins/Messaging/MessagingDebugger/Source/MessagingDebugger/Private/Models/MessagingDebuggerEndpointFilter.h b/Engine/Plugins/Messaging/MessagingDebugger/Source/MessagingDebugger/Private/Models/MessagingDebuggerEndpointFilter.h
--- a/Engine/Plugins/Messaging/MessagingDebugger/Source/MessagingDebugger/Private/Models/MessagingDebuggerEndpointFilter.h
+++ b/Engine/Plugins/Messaging/MessagingDebugger/Source/MessagingDebugger/Private/Models/MessagingDebuggerEndpointFilter.h
@@ -26,7 +26,7 @@ public:
 		// filter by search string
 		TArray<FString> FilterSubstrings;
 		
-		if (FilterString.ParseIntoArray(&FilterSubstrings, TEXT(" "), true) > 0)
+		if (FilterString.ParseIntoArray(FilterSubstrings, TEXT(" "), true) > 0)
 		{
 			for (int32 SubstringIndex = 0; SubstringIndex < FilterSubstrings.Num(); ++SubstringIndex)
 			{
diff --git a/Engine/Plugins/Messaging/MessagingDebugger/Source/MessagingDebugger/Private/Models/MessagingDebuggerTypeFilter.h b/Engine/Plugins/Messaging/MessagingDebugger/Source/MessagingDebugger/Private/Models/MessagingDebuggerTypeFilter.h
--- a/Engine/Plugins/Messaging/MessagingDebugger/Source/MessagingDebugger/Private/Models/MessagingDebuggerTypeFilter.h
+++ b/Engine/Plugins/Messaging/MessagingDebugger/Source/MessagingDebugger/Private/Models/MessagingDebuggerTypeFilter.h
@@ -26,7 +26,7 @@ public:
 		// filter by search string
 		TArray<FString> FilterSubstrings;
 		
-		if (FilterString.ParseIntoArray(&FilterSubstrings, TEXT(" "), true) > 0)
+		if (FilterString.ParseIntoArray(FilterSubstrings, TEXT(" "), true) > 0)
 		{
 			FString TypeName = TypeInfo->TypeName.ToString();
 
diff --git a/Engine/Source/Developer/Android/AndroidDeviceDetection/Private/AndroidDeviceDetectionModule.cpp b/Engine/Source/Developer/Android/AndroidDeviceDetection/Private/AndroidDeviceDetectionModule.cpp
--- a/Engine/Source/Developer/Android/AndroidDeviceDetection/Private/AndroidDeviceDetectionModule.cpp
+++ b/Engine/Source/Developer/Android/AndroidDeviceDetection/Private/AndroidDeviceDetectionModule.cpp
@@ -117,7 +117,7 @@ private:
 		// separate out each line
 		TArray<FString> DeviceStrings;
 		StdOut = StdOut.Replace(TEXT("\r"), TEXT("\n"));
-		StdOut.ParseIntoArray(&DeviceStrings, TEXT("\n"), true);
+		StdOut.ParseIntoArray(DeviceStrings, TEXT("\n"), true);
 
 		// a list containing all devices found this time, so we can remove anything not in this list
 		TArray<FString> CurrentlyConnectedDevices;
@@ -327,7 +327,7 @@ public:
 				delete FileReader;
 
 				TArray<FString> Lines;
-				FString(ANSI_TO_TCHAR(AnsiContents)).ParseIntoArrayLines(&Lines);
+				FString(ANSI_TO_TCHAR(AnsiContents)).ParseIntoArrayLines(Lines);
 				FMemory::Free(AnsiContents);
 
 				for (int32 Index = 0; Index < Lines.Num(); Index++)
diff --git a/Engine/Source/Developer/AssetTools/Private/SPackageReportDialog.cpp b/Engine/Source/Developer/AssetTools/Private/SPackageReportDialog.cpp
--- a/Engine/Source/Developer/AssetTools/Private/SPackageReportDialog.cpp
+++ b/Engine/Source/Developer/AssetTools/Private/SPackageReportDialog.cpp
@@ -28,7 +28,7 @@ FPackageReportNode::FPackageReportNode(const FString& InNodeName, bool InIsFolde
 void FPackageReportNode::AddPackage(const FString& PackageName)
 {
 	TArray<FString> PathElements;
-	PackageName.ParseIntoArray(&PathElements, TEXT("/"), /*InCullEmpty=*/true);
+	PackageName.ParseIntoArray(PathElements, TEXT("/"), /*InCullEmpty=*/true);
 
 	return AddPackage_Recursive(PathElements);
 }
diff --git a/Engine/Source/Developer/CollectionManager/Private/Collection.cpp b/Engine/Source/Developer/CollectionManager/Private/Collection.cpp
--- a/Engine/Source/Developer/CollectionManager/Private/Collection.cpp
+++ b/Engine/Source/Developer/CollectionManager/Private/Collection.cpp
@@ -53,7 +53,7 @@ bool FCollection::LoadFromFile(const FString& InFilename, bool InUseSCC)
 	// Normalize line endings and parse into array
 	TArray<FString> FileContents;
 	FullFileContentsString.ReplaceInline(TEXT("\r"), TEXT(""));
-	FullFileContentsString.ParseIntoArray(&FileContents, TEXT("\n"), /*bCullEmpty=*/false);
+	FullFileContentsString.ParseIntoArray(FileContents, TEXT("\n"), /*bCullEmpty=*/false);
 
 	if ( FileContents.Num() == 0 )
 	{
diff --git a/Engine/Source/Developer/CrashDebugHelper/Private/CrashDebugHelper.cpp b/Engine/Source/Developer/CrashDebugHelper/Private/CrashDebugHelper.cpp
--- a/Engine/Source/Developer/CrashDebugHelper/Private/CrashDebugHelper.cpp
+++ b/Engine/Source/Developer/CrashDebugHelper/Private/CrashDebugHelper.cpp
@@ -399,7 +399,7 @@ bool ICrashDebugHelper::ReadSourceFile( const TCHAR* InFilename, TArray<FString>
 	if( FFileHelper::LoadFileToString( Line, InFilename ) )
 	{
 		Line = Line.Replace( TEXT( "\r" ), TEXT( "" ) );
-		Line.ParseIntoArray( &OutStrings, TEXT( "\n" ), false );
+		Line.ParseIntoArray( OutStrings, TEXT( "\n" ), false );
 		
 		return true;
 	}
diff --git a/Engine/Source/Developer/CrashDebugHelper/Private/Windows/WindowsPlatformStackWalkExt.cpp b/Engine/Source/Developer/CrashDebugHelper/Private/Windows/WindowsPlatformStackWalkExt.cpp
--- a/Engine/Source/Developer/CrashDebugHelper/Private/Windows/WindowsPlatformStackWalkExt.cpp
+++ b/Engine/Source/Developer/CrashDebugHelper/Private/Windows/WindowsPlatformStackWalkExt.cpp
@@ -88,7 +88,7 @@ FString FWindowsPlatformStackWalkExt::ExtractRelativePath( const TCHAR* BaseName
 	FullPath = FullPath.Replace( TEXT( "\\" ), TEXT( "/" ) );
 
 	TArray<FString> Components;
-	int32 Count = FullPath.ParseIntoArray( &Components, TEXT( "/" ), true );
+	int32 Count = FullPath.ParseIntoArray( Components, TEXT( "/" ), true );
 	FullPath = TEXT( "" );
 
 	for( int32 Index = 0; Index < Count; Index++ )
@@ -369,7 +369,7 @@ void FWindowsPlatformStackWalkExt::SetSymbolPathsFromModules()
 
 	Symbol->GetSymbolPathWide( SymbolPath, ARRAY_COUNT(SymbolPath), NULL );
 	TArray<FString> SymbolPaths;
-	FString( SymbolPath ).ParseIntoArray(&SymbolPaths, TEXT(";"), true );
+	FString( SymbolPath ).ParseIntoArray(SymbolPaths, TEXT(";"), true );
 
 	UE_LOG( LogCrashDebugHelper, Log, TEXT( "Symbol paths" ) );
 	for( const auto& It : SymbolPaths )
@@ -379,7 +379,7 @@ void FWindowsPlatformStackWalkExt::SetSymbolPathsFromModules()
 
 	Symbol->GetImagePathWide( SymbolPath, ARRAY_COUNT( SymbolPath ), NULL );
 	TArray<FString> ImagePaths;
-	FString( SymbolPath ).ParseIntoArray( &ImagePaths, TEXT( ";" ), true );
+	FString( SymbolPath ).ParseIntoArray( ImagePaths, TEXT( ";" ), true );
 	
 	UE_LOG( LogCrashDebugHelper, Log, TEXT( "Image paths" ) );
 	for( const auto& It : ImagePaths )
diff --git a/Engine/Source/Developer/DesktopPlatform/Private/Linux/DesktopPlatformLinux.cpp b/Engine/Source/Developer/DesktopPlatform/Private/Linux/DesktopPlatformLinux.cpp
--- a/Engine/Source/Developer/DesktopPlatform/Private/Linux/DesktopPlatformLinux.cpp
+++ b/Engine/Source/Developer/DesktopPlatform/Private/Linux/DesktopPlatformLinux.cpp
@@ -127,7 +127,7 @@ bool FDesktopPlatformLinux::FileDialogShared(bool bSave, const void* ParentWindo
 	// It is formatted as follows: Pair1String1|Pair1String2|Pair2String1|Pair2String2
 	// where the second string in the pair is the extension.  To get the clean extensions we only care about the second string in the pair
 	TArray<FString> UnformattedExtensions;
-	FileTypes.ParseIntoArray( &UnformattedExtensions, TEXT("|"), true );
+	FileTypes.ParseIntoArray( UnformattedExtensions, TEXT("|"), true );
 	for( int32 ExtensionIndex = 1; ExtensionIndex < UnformattedExtensions.Num(); ExtensionIndex += 2)
 	{
 		const FString& Extension = UnformattedExtensions[ExtensionIndex];
diff --git a/Engine/Source/Developer/DesktopPlatform/Private/Mac/DesktopPlatformMac.cpp b/Engine/Source/Developer/DesktopPlatform/Private/Mac/DesktopPlatformMac.cpp
--- a/Engine/Source/Developer/DesktopPlatform/Private/Mac/DesktopPlatformMac.cpp
+++ b/Engine/Source/Developer/DesktopPlatform/Private/Mac/DesktopPlatformMac.cpp
@@ -473,7 +473,7 @@ bool FDesktopPlatformMac::FileDialogShared(bool bSave, const void* ParentWindowH
 			[Panel setAccessoryView: AccessoryView];
 
 			TArray<FString> FileTypesArray;
-			int32 NumFileTypes = FileTypes.ParseIntoArray(&FileTypesArray, TEXT("|"), true);
+			int32 NumFileTypes = FileTypes.ParseIntoArray(FileTypesArray, TEXT("|"), true);
 
 			NSMutableArray* AllowedFileTypes = [NSMutableArray arrayWithCapacity: NumFileTypes];
 
diff --git a/Engine/Source/Developer/DesktopPlatform/Private/Windows/DesktopPlatformWindows.cpp b/Engine/Source/Developer/DesktopPlatform/Private/Windows/DesktopPlatformWindows.cpp
--- a/Engine/Source/Developer/DesktopPlatform/Private/Windows/DesktopPlatformWindows.cpp
+++ b/Engine/Source/Developer/DesktopPlatform/Private/Windows/DesktopPlatformWindows.cpp
@@ -247,7 +247,7 @@ bool FDesktopPlatformWindows::FileDialogShared(bool bSave, const void* ParentWin
 	// It is formatted as follows: Pair1String1|Pair1String2|Pair2String1|Pair2String2
 	// where the second string in the pair is the extension.  To get the clean extensions we only care about the second string in the pair
 	TArray<FString> UnformattedExtensions;
-	FileTypes.ParseIntoArray( &UnformattedExtensions, TEXT("|"), true );
+	FileTypes.ParseIntoArray( UnformattedExtensions, TEXT("|"), true );
 	for( int32 ExtensionIndex = 1; ExtensionIndex < UnformattedExtensions.Num(); ExtensionIndex += 2)
 	{
 		const FString& Extension = UnformattedExtensions[ExtensionIndex];
diff --git a/Engine/Source/Developer/FunctionalTesting/Private/FunctionalTestingManager.cpp b/Engine/Source/Developer/FunctionalTesting/Private/FunctionalTestingManager.cpp
--- a/Engine/Source/Developer/FunctionalTesting/Private/FunctionalTestingManager.cpp
+++ b/Engine/Source/Developer/FunctionalTesting/Private/FunctionalTestingManager.cpp
@@ -293,7 +293,7 @@ bool UFunctionalTestingManager::RunFirstValidTest()
 		{
 			TArray<FString> TestParams;
 			const FString SingleTestReproString = TestReproStrings.Pop(/*bAllowShrinking=*/false);
-			SingleTestReproString.ParseIntoArray(&TestParams, TEXT("#"), /*InCullEmpty=*/true);
+			SingleTestReproString.ParseIntoArray(TestParams, TEXT("#"), /*InCullEmpty=*/true);
 			
 			if (TestParams.Num() == 0)
 			{
@@ -366,7 +366,7 @@ void UFunctionalTestingManager::SetReproString(FString ReproString)
 	StartingReproString = ReproString;
 	if (ReproString.IsEmpty() == false)
 	{
-		ReproString.ParseIntoArray(&TestReproStrings, FFunctionalTesting::ReproStringTestSeparator, /*InCullEmpty=*/true);
+		ReproString.ParseIntoArray(TestReproStrings, FFunctionalTesting::ReproStringTestSeparator, /*InCullEmpty=*/true);
 	}
 }
 
diff --git a/Engine/Source/Developer/HTML5/HTML5TargetPlatform/Private/HTML5TargetPlatform.cpp b/Engine/Source/Developer/HTML5/HTML5TargetPlatform/Private/HTML5TargetPlatform.cpp
--- a/Engine/Source/Developer/HTML5/HTML5TargetPlatform/Private/HTML5TargetPlatform.cpp
+++ b/Engine/Source/Developer/HTML5/HTML5TargetPlatform/Private/HTML5TargetPlatform.cpp
@@ -353,7 +353,7 @@ void FHTML5TargetPlatform::RefreshAvailableDevices()
 				FString DevicePath;
 				It.Value.RemoveFromStart(TEXT("("));
 				It.Value.RemoveFromEnd(TEXT(")"));
-				It.Value.ParseIntoArray(&ValueArray, TEXT(","), 1);
+				It.Value.ParseIntoArray(ValueArray, TEXT(","), 1);
 				for (auto& Value : ValueArray)
 				{
 					if (Value.StartsWith(TEXT("DeviceName=")))
diff --git a/Engine/Source/Developer/LauncherServices/Private/Launcher/LauncherWorker.cpp b/Engine/Source/Developer/LauncherServices/Private/Launcher/LauncherWorker.cpp
--- a/Engine/Source/Developer/LauncherServices/Private/Launcher/LauncherWorker.cpp
+++ b/Engine/Source/Developer/LauncherServices/Private/Launcher/LauncherWorker.cpp
@@ -46,7 +46,7 @@ uint32 FLauncherWorker::Run( )
 			// process the string to break it up in to lines
 			Line += NewLine;
 			TArray<FString> StringArray;
-			int32 count = Line.ParseIntoArray(&StringArray, TEXT("\n"), true);
+			int32 count = Line.ParseIntoArray(StringArray, TEXT("\n"), true);
 			if (count > 1)
 			{
 				for (int32 Index = 0; Index < count-1; ++Index)
@@ -70,7 +70,7 @@ uint32 FLauncherWorker::Run( )
 				// process the string to break it up in to lines
 				Line += NewLine;
 				TArray<FString> StringArray;
-				int32 count = Line.ParseIntoArray(&StringArray, TEXT("\n"), true);
+				int32 count = Line.ParseIntoArray(StringArray, TEXT("\n"), true);
 				if (count > 1)
 				{
 					for (int32 Index = 0; Index < count-1; ++Index)
diff --git a/Engine/Source/Developer/LauncherServices/Private/Profiles/LauncherProfileManager.cpp b/Engine/Source/Developer/LauncherServices/Private/Profiles/LauncherProfileManager.cpp
--- a/Engine/Source/Developer/LauncherServices/Private/Profiles/LauncherProfileManager.cpp
+++ b/Engine/Source/Developer/LauncherServices/Private/Profiles/LauncherProfileManager.cpp
@@ -406,7 +406,7 @@ ILauncherDeviceGroupPtr FLauncherProfileManager::ParseDeviceGroup( const FString
 		Result = MakeShareable(new FLauncherDeviceGroup(GroupId, GroupName));
 
 		TArray<FString> DeviceList;
-		DevicesString.ParseIntoArray(&DeviceList, TEXT(", "), false);
+		DevicesString.ParseIntoArray(DeviceList, TEXT(", "), false);
 
 		for (int32 Index = 0; Index < DeviceList.Num(); ++Index)
 		{
diff --git a/Engine/Source/Developer/Profiler/Private/ProfilerSession.cpp b/Engine/Source/Developer/Profiler/Private/ProfilerSession.cpp
--- a/Engine/Source/Developer/Profiler/Private/ProfilerSession.cpp
+++ b/Engine/Source/Developer/Profiler/Private/ProfilerSession.cpp
@@ -1204,7 +1204,7 @@ static FString EncodeCallstack( const TArray<FName>& Callstack )
 
 static void DecodeCallstack( const FName& EncodedCallstack, TArray<FString>& out_DecodedCallstack )
 {
-	EncodedCallstack.ToString().ParseIntoArray( &out_DecodedCallstack, CallstackSeparator, true );
+	EncodedCallstack.ToString().ParseIntoArray( out_DecodedCallstack, CallstackSeparator, true );
 }
 
 /** Simple allocation info. Assumes the sequence tag are unique and doesn't wrap. */
diff --git a/Engine/Source/Developer/ShaderFormatOpenGL/Private/OpenGLShaderCompiler.cpp b/Engine/Source/Developer/ShaderFormatOpenGL/Private/OpenGLShaderCompiler.cpp
--- a/Engine/Source/Developer/ShaderFormatOpenGL/Private/OpenGLShaderCompiler.cpp
+++ b/Engine/Source/Developer/ShaderFormatOpenGL/Private/OpenGLShaderCompiler.cpp
@@ -1416,7 +1416,7 @@ static void PrecompileShader(FShaderCompilerOutput& ShaderOutput, const FShaderC
 					RawCompileLog.AddZeroed(LogLength);
 					glGetShaderInfoLog(Shader, LogLength, /*OutLength=*/ NULL, RawCompileLog.GetData());
 					CompileLog = ANSI_TO_TCHAR(RawCompileLog.GetData());
-					CompileLog.ParseIntoArray(&LogLines, TEXT("\n"), true);
+					CompileLog.ParseIntoArray(LogLines, TEXT("\n"), true);
 
 					for (int32 Line = 0; Line < LogLines.Num(); ++Line)
 					{
@@ -1773,7 +1773,7 @@ void CompileShader_Windows_OGL(const FShaderCompilerInput& Input,FShaderCompiler
 
 			FString Tmp = ANSI_TO_TCHAR(ErrorLog);
 			TArray<FString> ErrorLines;
-			Tmp.ParseIntoArray(&ErrorLines, TEXT("\n"), true);
+			Tmp.ParseIntoArray(ErrorLines, TEXT("\n"), true);
 			for (int32 LineIndex = 0; LineIndex < ErrorLines.Num(); ++LineIndex)
 			{
 				const FString& Line = ErrorLines[LineIndex];
diff --git a/Engine/Source/Developer/ShaderPreprocessor/Private/PreprocessorPrivate.h b/Engine/Source/Developer/ShaderPreprocessor/Private/PreprocessorPrivate.h
--- a/Engine/Source/Developer/ShaderPreprocessor/Private/PreprocessorPrivate.h
+++ b/Engine/Source/Developer/ShaderPreprocessor/Private/PreprocessorPrivate.h
@@ -41,7 +41,7 @@ static bool ParseMcppErrors(TArray<FShaderCompilerError>& OutErrors, const FStri
 	if (McppErrors.Len() > 0)
 	{
 		TArray<FString> Lines;
-		McppErrors.ParseIntoArray(&Lines, TEXT("\n"), true);
+		McppErrors.ParseIntoArray(Lines, TEXT("\n"), true);
 		for (int32 LineIndex = 0; LineIndex < Lines.Num(); ++LineIndex)
 		{
 			const FString& Line = Lines[LineIndex];
diff --git a/Engine/Source/Developer/SourceControl/Private/SourceControlSettings.cpp b/Engine/Source/Developer/SourceControl/Private/SourceControlSettings.cpp
--- a/Engine/Source/Developer/SourceControl/Private/SourceControlSettings.cpp
+++ b/Engine/Source/Developer/SourceControl/Private/SourceControlSettings.cpp
@@ -48,7 +48,7 @@ void FSourceControlSettings::LoadSettings()
 	{
 		FString& Switch = Switches[SwitchIdx];
 		TArray<FString> SplitSwitch;
-		if (2 == Switch.ParseIntoArray(&SplitSwitch, TEXT("="), true))
+		if (2 == Switch.ParseIntoArray(SplitSwitch, TEXT("="), true))
 		{
 			SwitchPairs.Add(SplitSwitch[0], SplitSwitch[1].TrimQuotes());
 			Switches.RemoveAt(SwitchIdx);
diff --git a/Engine/Source/Developer/SourceControl/Private/Tests/SourceControlTests.cpp b/Engine/Source/Developer/SourceControl/Private/Tests/SourceControlTests.cpp
--- a/Engine/Source/Developer/SourceControl/Private/Tests/SourceControlTests.cpp
+++ b/Engine/Source/Developer/SourceControl/Private/Tests/SourceControlTests.cpp
@@ -226,7 +226,7 @@ bool FCheckOutTest::RunTest(const FString& Parameters)
 	// parameter is the provider we want to use followed by the filename
 	const FString Delimiter(TEXT(" "));
 	TArray<FString> ParamArray;
-	Parameters.ParseIntoArray(&ParamArray, *Delimiter, true);
+	Parameters.ParseIntoArray(ParamArray, *Delimiter, true);
 	ensure(ParamArray.Num() == 2);
 
 	// check to see if we should restore the read only status after this test
@@ -377,7 +377,7 @@ bool FMarkForAddTest::RunTest(const FString& Parameters)
 	// parameter is the provider we want to use followed by the filename
 	const FString Delimiter(TEXT(" "));
 	TArray<FString> ParamArray;
-	Parameters.ParseIntoArray(&ParamArray, *Delimiter, true);
+	Parameters.ParseIntoArray(ParamArray, *Delimiter, true);
 	ensure(ParamArray.Num() == 2);
 
 	ADD_LATENT_AUTOMATION_COMMAND(FSetProviderLatentCommand(FName(*ParamArray[0])));
@@ -450,7 +450,7 @@ bool FDeleteTest::RunTest(const FString& Parameters)
 	// parameter is the provider we want to use followed by the filename
 	const FString Delimiter(TEXT(" "));
 	TArray<FString> ParamArray;
-	Parameters.ParseIntoArray(&ParamArray, *Delimiter, true);
+	Parameters.ParseIntoArray(ParamArray, *Delimiter, true);
 	ensure(ParamArray.Num() == 2);
 
 	// check to see if we should restore the read only status after this test
@@ -551,7 +551,7 @@ bool FCheckInTest::RunTest(const FString& Parameters)
 	// parameter is the provider we want to use followed by the filename
 	const FString Delimiter(TEXT(" "));
 	TArray<FString> ParamArray;
-	Parameters.ParseIntoArray(&ParamArray, *Delimiter, true);
+	Parameters.ParseIntoArray(ParamArray, *Delimiter, true);
 	ensure(ParamArray.Num() == 2);
 
 	// check to see if we should restore the read only status after this test
@@ -621,7 +621,7 @@ bool FSyncTest::RunTest(const FString& Parameters)
 	// parameter is the provider we want to use followed by the filename
 	const FString Delimiter(TEXT(" "));
 	TArray<FString> ParamArray;
-	Parameters.ParseIntoArray(&ParamArray, *Delimiter, true);
+	Parameters.ParseIntoArray(ParamArray, *Delimiter, true);
 	ensure(ParamArray.Num() == 2);
 
 	ADD_LATENT_AUTOMATION_COMMAND(FSetProviderLatentCommand(FName(*ParamArray[0])));
@@ -644,7 +644,7 @@ bool FRevertTest::RunTest(const FString& Parameters)
 	// parameter is the provider we want to use followed by the filename
 	const FString Delimiter(TEXT(" "));
 	TArray<FString> ParamArray;
-	Parameters.ParseIntoArray(&ParamArray, *Delimiter, true);
+	Parameters.ParseIntoArray(ParamArray, *Delimiter, true);
 	ensure(ParamArray.Num() == 2);
 
 	ADD_LATENT_AUTOMATION_COMMAND(FSetProviderLatentCommand(FName(*ParamArray[0])));
@@ -739,7 +739,7 @@ bool FUpdateStatusTest::RunTest(const FString& Parameters)
 	// parameter is the provider we want to use followed by the filename
 	const FString Delimiter(TEXT(" "));
 	TArray<FString> ParamArray;
-	Parameters.ParseIntoArray(&ParamArray, *Delimiter, true);
+	Parameters.ParseIntoArray(ParamArray, *Delimiter, true);
 	ensure(ParamArray.Num() == 2);
 
 	ADD_LATENT_AUTOMATION_COMMAND(FSetProviderLatentCommand(FName(*ParamArray[0])));
@@ -823,7 +823,7 @@ bool FGetLabelTest::RunTest(const FString& Parameters)
 	// parameter is the provider we want to use followed by the label
 	const FString Delimiter(TEXT(" "));
 	TArray<FString> ParamArray;
-	Parameters.ParseIntoArray(&ParamArray, *Delimiter, true);
+	Parameters.ParseIntoArray(ParamArray, *Delimiter, true);
 	ensure(ParamArray.Num() == 2);
 
 	TArray<FString> FilesToGet;
@@ -868,7 +868,7 @@ bool FSyncLabelTest::RunTest(const FString& Parameters)
 	// parameter is the provider we want to use followed by the label
 	const FString Delimiter(TEXT(" "));
 	TArray<FString> ParamArray;
-	Parameters.ParseIntoArray(&ParamArray, *Delimiter, true);
+	Parameters.ParseIntoArray(ParamArray, *Delimiter, true);
 	ensure(ParamArray.Num() == 2);
 
 	TArray<FString> FilesToGet;
@@ -934,7 +934,7 @@ bool FGetRevisionTest::RunTest(const FString& Parameters)
 	// parameter is the provider we want to use followed by the filename
 	const FString Delimiter(TEXT(" "));
 	TArray<FString> ParamArray;
-	Parameters.ParseIntoArray(&ParamArray, *Delimiter, true);
+	Parameters.ParseIntoArray(ParamArray, *Delimiter, true);
 	ensure(ParamArray.Num() == 2);
 
 	ADD_LATENT_AUTOMATION_COMMAND(FSetProviderLatentCommand(FName(*ParamArray[0])));
diff --git a/Engine/Source/Developer/TargetPlatform/Private/TargetPlatformManagerModule.cpp b/Engine/Source/Developer/TargetPlatform/Private/TargetPlatformManagerModule.cpp
--- a/Engine/Source/Developer/TargetPlatform/Private/TargetPlatformManagerModule.cpp
+++ b/Engine/Source/Developer/TargetPlatform/Private/TargetPlatformManagerModule.cpp
@@ -186,7 +186,7 @@ public:
 				{
 					TArray<FString> PlatformNames;
 
-					PlatformStr.ParseIntoArray(&PlatformNames, TEXT("+"), true);
+					PlatformStr.ParseIntoArray(PlatformNames, TEXT("+"), true);
 
 					// for nicer user response
 					FString AvailablePlatforms;
@@ -665,7 +665,7 @@ protected:
 			InstallManifestFile->Serialize(FileMem, FileSize);
 
 			FString FileAsString(ANSI_TO_TCHAR(FileMem));
-			FileAsString.ParseIntoArrayLines(&FileLines);
+			FileAsString.ParseIntoArrayLines(FileLines);
 
 			FMemory::Free(FileMem);
 			InstallManifestFile->Close();
@@ -707,7 +707,7 @@ protected:
 				EnvVarFile->Serialize(FileMem, FileSize);
 
 				FString FileAsString(ANSI_TO_TCHAR(FileMem));
-				FileAsString.ParseIntoArrayLines(&FileLines);
+				FileAsString.ParseIntoArrayLines(FileLines);
 
 				FMemory::Free(FileMem);
 				EnvVarFile->Close();				
@@ -758,7 +758,7 @@ protected:
 
 			const TCHAR* PathDelimiter = FPlatformMisc::GetPathVarDelimiter();
 			TArray<FString> PathVars;
-			OrigPathVar.ParseIntoArray(&PathVars, PathDelimiter, true);
+			OrigPathVar.ParseIntoArray(PathVars, PathDelimiter, true);
 
 			TArray<FString> ModifiedPathVars;
 			ModifiedPathVars = PathVars;
@@ -865,7 +865,7 @@ protected:
 		}
 
 		TArray<FString> PlatArray;
-		SDKStatusMessage.ParseIntoArrayWS(&PlatArray);
+		SDKStatusMessage.ParseIntoArrayWS(PlatArray);
 		for (int Index = 0; Index < PlatArray.Num()-2; ++Index)
 		{
             FString Item = PlatArray[Index];
diff --git a/Engine/Source/Developer/Windows/ShaderFormatD3D/Private/D3D11ShaderCompiler.cpp b/Engine/Source/Developer/Windows/ShaderFormatD3D/Private/D3D11ShaderCompiler.cpp
--- a/Engine/Source/Developer/Windows/ShaderFormatD3D/Private/D3D11ShaderCompiler.cpp
+++ b/Engine/Source/Developer/Windows/ShaderFormatD3D/Private/D3D11ShaderCompiler.cpp
@@ -50,7 +50,7 @@ static void D3D11FilterShaderCompileWarnings(const FString& CompileWarnings, TAr
 {
 	TArray<FString> WarningArray;
 	FString OutWarningString = TEXT("");
-	CompileWarnings.ParseIntoArray(&WarningArray, TEXT("\n"), true);
+	CompileWarnings.ParseIntoArray(WarningArray, TEXT("\n"), true);
 	
 	//go through each warning line
 	for (int32 WarningIndex = 0; WarningIndex < WarningArray.Num(); WarningIndex++)
diff --git a/Engine/Source/Developer/iOS/IOSPlatformEditor/Private/IOSTargetSettingsCustomization.cpp b/Engine/Source/Developer/iOS/IOSPlatformEditor/Private/IOSTargetSettingsCustomization.cpp
--- a/Engine/Source/Developer/iOS/IOSPlatformEditor/Private/IOSTargetSettingsCustomization.cpp
+++ b/Engine/Source/Developer/iOS/IOSPlatformEditor/Private/IOSTargetSettingsCustomization.cpp
@@ -117,7 +117,7 @@ void FIOSTargetSettingsCustomization::UpdateStatus()
 
 		// Now split up the log into multiple lines
 		TArray<FString> LogLines;
-		OutputMessage.ParseIntoArray(&LogLines, LINE_TERMINATOR, true);
+		OutputMessage.ParseIntoArray(LogLines, LINE_TERMINATOR, true);
 		
 		// format of the line being read here!!
 		bool bCerts = false;
@@ -125,7 +125,7 @@ void FIOSTargetSettingsCustomization::UpdateStatus()
 		{
 			FString& Line = LogLines[Index];
 			TArray<FString> Fields;
-			Line.ParseIntoArray(&Fields, TEXT(","), true);
+			Line.ParseIntoArray(Fields, TEXT(","), true);
 			if (Line.Contains(TEXT("CERTIFICATE-"), ESearchCase::CaseSensitive))
 			{
 				CertificatePtr Cert = MakeShareable<FCertificate>(new FCertificate());
diff --git a/Engine/Source/Developer/iOS/IOSTargetPlatform/Private/Windows/IOSDeviceHelperWindows.cpp b/Engine/Source/Developer/iOS/IOSTargetPlatform/Private/Windows/IOSDeviceHelperWindows.cpp
--- a/Engine/Source/Developer/iOS/IOSTargetPlatform/Private/Windows/IOSDeviceHelperWindows.cpp
+++ b/Engine/Source/Developer/iOS/IOSTargetPlatform/Private/Windows/IOSDeviceHelperWindows.cpp
@@ -152,7 +152,7 @@ private:
 		// separate out each line
 		TArray<FString> DeviceStrings;
 		StdOut = StdOut.Replace(TEXT("\r"), TEXT("\n"));
-		StdOut.ParseIntoArray(&DeviceStrings, TEXT("\n"), true);
+		StdOut.ParseIntoArray(DeviceStrings, TEXT("\n"), true);
 
 		TArray<FString> CurrentDeviceIds;
 		bool FoundList = false;
diff --git a/Engine/Source/Developer/iOS/MetalShaderFormat/Private/MetalShaderCompiler.cpp b/Engine/Source/Developer/iOS/MetalShaderFormat/Private/MetalShaderCompiler.cpp
--- a/Engine/Source/Developer/iOS/MetalShaderFormat/Private/MetalShaderCompiler.cpp
+++ b/Engine/Source/Developer/iOS/MetalShaderFormat/Private/MetalShaderCompiler.cpp
@@ -960,7 +960,7 @@ void CompileShader_Metal(const FShaderCompilerInput& Input,FShaderCompilerOutput
 		{
 			FString Tmp = ANSI_TO_TCHAR(ErrorLog);
 			TArray<FString> ErrorLines;
-			Tmp.ParseIntoArray(&ErrorLines, TEXT("\n"), true);
+			Tmp.ParseIntoArray(ErrorLines, TEXT("\n"), true);
 			for (int32 LineIndex = 0; LineIndex < ErrorLines.Num(); ++LineIndex)
 			{
 				const FString& Line = ErrorLines[LineIndex];
diff --git a/Engine/Source/Editor/AnimGraph/Private/AnimGraphNode_SkeletalControlBase.cpp b/Engine/Source/Editor/AnimGraph/Private/AnimGraphNode_SkeletalControlBase.cpp
--- a/Engine/Source/Editor/AnimGraph/Private/AnimGraphNode_SkeletalControlBase.cpp
+++ b/Engine/Source/Editor/AnimGraph/Private/AnimGraphNode_SkeletalControlBase.cpp
@@ -293,7 +293,7 @@ void UAnimGraphNode_SkeletalControlBase::GetDefaultValue(const FString& UpdateDe
 				//Parse string to split its contents separated by ','
 				DefaultString.Trim();
 				DefaultString.TrimTrailing();
-				DefaultString.ParseIntoArray(&ResultString, TEXT(","), true);
+				DefaultString.ParseIntoArray(ResultString, TEXT(","), true);
 
 				check(ResultString.Num() == 3);
 
diff --git a/Engine/Source/Editor/BehaviorTreeEditor/Private/FindInBT.cpp b/Engine/Source/Editor/BehaviorTreeEditor/Private/FindInBT.cpp
--- a/Engine/Source/Editor/BehaviorTreeEditor/Private/FindInBT.cpp
+++ b/Engine/Source/Editor/BehaviorTreeEditor/Private/FindInBT.cpp
@@ -169,7 +169,7 @@ void SFindInBT::OnSearchTextCommitted(const FText& Text, ETextCommit::Type Commi
 void SFindInBT::InitiateSearch()
 {
 	TArray<FString> Tokens;
-	SearchValue.ParseIntoArray(&Tokens, TEXT(" "), true);
+	SearchValue.ParseIntoArray(Tokens, TEXT(" "), true);
 
 	for (auto It(ItemsFound.CreateIterator()); It; ++It)
 	{
diff --git a/Engine/Source/Editor/BlueprintGraph/Private/BlueprintNodeSignature.cpp b/Engine/Source/Editor/BlueprintGraph/Private/BlueprintNodeSignature.cpp
--- a/Engine/Source/Editor/BlueprintGraph/Private/BlueprintNodeSignature.cpp
+++ b/Engine/Source/Editor/BlueprintGraph/Private/BlueprintNodeSignature.cpp
@@ -54,7 +54,7 @@ FBlueprintNodeSignature::FBlueprintNodeSignature(FString const& UserString)
 	SanitizedSignature.RemoveFromEnd(SignatureClosingStr);
 
 	TArray<FString> SignatureElements;
-	SanitizedSignature.ParseIntoArray(&SignatureElements, *SignatureElementDelim, 1);
+	SanitizedSignature.ParseIntoArray(SignatureElements, *SignatureElementDelim, 1);
 
 	for (FString& SignatureElement : SignatureElements)
 	{
diff --git a/Engine/Source/Editor/BlueprintGraph/Private/EdGraphSchema_K2.cpp b/Engine/Source/Editor/BlueprintGraph/Private/EdGraphSchema_K2.cpp
--- a/Engine/Source/Editor/BlueprintGraph/Private/EdGraphSchema_K2.cpp
+++ b/Engine/Source/Editor/BlueprintGraph/Private/EdGraphSchema_K2.cpp
@@ -796,7 +796,7 @@ void UEdGraphSchema_K2::GetAutoEmitTermParameters(const UFunction* Function, TAr
 	if( Function->HasMetaData(FBlueprintMetadata::MD_AutoCreateRefTerm) )
 	{
 		FString MetaData = Function->GetMetaData(FBlueprintMetadata::MD_AutoCreateRefTerm);
-		MetaData.ParseIntoArray(&AutoEmitParameterNames, TEXT(","), true);
+		MetaData.ParseIntoArray(AutoEmitParameterNames, TEXT(","), true);
 	}
 }
 
@@ -5475,7 +5475,7 @@ void UEdGraphSchema_K2::SplitPin(UEdGraphPin* Pin) const
 		if (   StructType == GetBaseStructure(TEXT("Vector"))
 			|| StructType == GetBaseStructure(TEXT("Rotator")))
 		{
-			Pin->DefaultValue.ParseIntoArray(&OriginalDefaults, TEXT(","), false);
+			Pin->DefaultValue.ParseIntoArray(OriginalDefaults, TEXT(","), false);
 			for (FString& Default : OriginalDefaults)
 			{
 				Default = FString::SanitizeFloat(FCString::Atof(*Default));
diff --git a/Engine/Source/Editor/BlueprintGraph/Private/K2ActionMenuBuilder.cpp b/Engine/Source/Editor/BlueprintGraph/Private/K2ActionMenuBuilder.cpp
--- a/Engine/Source/Editor/BlueprintGraph/Private/K2ActionMenuBuilder.cpp
+++ b/Engine/Source/Editor/BlueprintGraph/Private/K2ActionMenuBuilder.cpp
@@ -1510,7 +1510,7 @@ void FK2ActionMenuBuilder::GetFuncNodesForClass(FGraphActionListBuilderBase& Lis
 				if(ObjProp->HasAllPropertyFlags(CPF_BlueprintVisible) && ExposeFuncCatsString.Len() > 0)
 				{
 					TArray<FString> ExposeFuncCats;
-					ExposeFuncCatsString.ParseIntoArray(&ExposeFuncCats, TEXT( "," ), true);
+					ExposeFuncCatsString.ParseIntoArray(ExposeFuncCats, TEXT( "," ), true);
 
 					// Look for functions in those categories
 					UClass* ObjClass = ObjProp->PropertyClass;
@@ -2056,7 +2056,7 @@ void FK2ActionMenuBuilder::GetEventsForBlueprint(FBlueprintPaletteListBuilder& A
 		if( Blueprint->ParentClass->HasMetaData(*ExclusionListKeyName) )
 		{
 			const FString ExcludedEventNameString = Blueprint->ParentClass->GetMetaData(*ExclusionListKeyName);
-			ExcludedEventNameString.ParseIntoArray(&ExcludedEventNames, TEXT(","), true);
+			ExcludedEventNameString.ParseIntoArray(ExcludedEventNames, TEXT(","), true);
 		}
 
 		// Start the name list by first adding all events already on the graph
diff --git a/Engine/Source/Editor/BlueprintGraph/Private/K2Node.cpp b/Engine/Source/Editor/BlueprintGraph/Private/K2Node.cpp
--- a/Engine/Source/Editor/BlueprintGraph/Private/K2Node.cpp
+++ b/Engine/Source/Editor/BlueprintGraph/Private/K2Node.cpp
@@ -784,7 +784,7 @@ void UK2Node::InitFieldRedirectMap()
 					FFieldRemapInfo OldFieldRemap;
 					{
 						TArray<FString> OldFieldPath;
-						OldFieldPathString.ParseIntoArray(&OldFieldPath, TEXT("."), true);
+						OldFieldPathString.ParseIntoArray(OldFieldPath, TEXT("."), true);
 
 						if (OldFieldPath.Num() == 1)
 						{
@@ -803,7 +803,7 @@ void UK2Node::InitFieldRedirectMap()
 					FFieldRemapInfo NewFieldRemap;
 					{
 						TArray<FString> NewFieldPath;
-						NewFieldPathString.ParseIntoArray(&NewFieldPath, TEXT("."), true);
+						NewFieldPathString.ParseIntoArray(NewFieldPath, TEXT("."), true);
 
 						if( NewFieldPath.Num() == 1 )
 						{
@@ -852,8 +852,8 @@ void UK2Node::InitFieldRedirectMap()
 
 					TArray<FString> OldParamValuesList;
 					TArray<FString> NewParamValuesList;
-					OldParamValues.ParseIntoArray(&OldParamValuesList, TEXT(";"), false);
-					NewParamValues.ParseIntoArray(&NewParamValuesList, TEXT(";"), false);
+					OldParamValues.ParseIntoArray(OldParamValuesList, TEXT(";"), false);
+					NewParamValues.ParseIntoArray(NewParamValuesList, TEXT(";"), false);
 
 					if (OldParamValuesList.Num() != NewParamValuesList.Num())
 					{
diff --git a/Engine/Source/Editor/BlueprintGraph/Private/K2Node_CallArrayFunction.cpp b/Engine/Source/Editor/BlueprintGraph/Private/K2Node_CallArrayFunction.cpp
--- a/Engine/Source/Editor/BlueprintGraph/Private/K2Node_CallArrayFunction.cpp
+++ b/Engine/Source/Editor/BlueprintGraph/Private/K2Node_CallArrayFunction.cpp
@@ -135,12 +135,12 @@ void UK2Node_CallArrayFunction::GetArrayPins(TArray< FArrayPropertyPinCombo >& O
 	check(TargetFunction);
 	FString ArrayPointerMetaData = TargetFunction->GetMetaData(TEXT("ArrayParm"));
 	TArray<FString> ArrayPinComboNames;
-	ArrayPointerMetaData.ParseIntoArray(&ArrayPinComboNames, TEXT(","), true);
+	ArrayPointerMetaData.ParseIntoArray(ArrayPinComboNames, TEXT(","), true);
 
 	for(auto Iter = ArrayPinComboNames.CreateConstIterator(); Iter; ++Iter)
 	{
 		TArray<FString> ArrayPinNames;
-		Iter->ParseIntoArray(&ArrayPinNames, TEXT("|"), true);
+		Iter->ParseIntoArray(ArrayPinNames, TEXT("|"), true);
 
 		FArrayPropertyPinCombo ArrayInfo;
 		ArrayInfo.ArrayPin = FindPin(ArrayPinNames[0]);
@@ -162,12 +162,12 @@ bool UK2Node_CallArrayFunction::IsWildcardProperty(UFunction* InArrayFunction, c
 	{
 		FString ArrayPointerMetaData = InArrayFunction->GetMetaData(TEXT("ArrayParm"));
 		TArray<FString> ArrayPinComboNames;
-		ArrayPointerMetaData.ParseIntoArray(&ArrayPinComboNames, TEXT(","), true);
+		ArrayPointerMetaData.ParseIntoArray(ArrayPinComboNames, TEXT(","), true);
 
 		for(auto Iter = ArrayPinComboNames.CreateConstIterator(); Iter; ++Iter)
 		{
 			TArray<FString> ArrayPinNames;
-			Iter->ParseIntoArray(&ArrayPinNames, TEXT("|"), true);
+			Iter->ParseIntoArray(ArrayPinNames, TEXT("|"), true);
 
 			if(ArrayPinNames[0] == InProperty->GetName())
 			{
@@ -187,7 +187,7 @@ void UK2Node_CallArrayFunction::GetArrayTypeDependentPins(TArray<UEdGraphPin*>&
 
 	const FString DependentPinMetaData = TargetFunction->GetMetaData(TEXT("ArrayTypeDependentParams"));
 	TArray<FString> TypeDependentPinNames;
-	DependentPinMetaData.ParseIntoArray(&TypeDependentPinNames, TEXT(","), true);
+	DependentPinMetaData.ParseIntoArray(TypeDependentPinNames, TEXT(","), true);
 
 	for(TArray<UEdGraphPin*>::TConstIterator it(Pins); it; ++it)
 	{
diff --git a/Engine/Source/Editor/BlueprintGraph/Private/K2Node_CallFunction.cpp b/Engine/Source/Editor/BlueprintGraph/Private/K2Node_CallFunction.cpp
--- a/Engine/Source/Editor/BlueprintGraph/Private/K2Node_CallFunction.cpp
+++ b/Engine/Source/Editor/BlueprintGraph/Private/K2Node_CallFunction.cpp
@@ -34,7 +34,7 @@ struct FCustomStructureParamHelper
 			FString MetaDataValue = Function->GetMetaData(GetCustomStructureParamName());
 			if (!MetaDataValue.IsEmpty())
 			{
-				MetaDataValue.ParseIntoArray(&OutNames, TEXT(","), true);
+				MetaDataValue.ParseIntoArray(OutNames, TEXT(","), true);
 			}
 		}
 	}
diff --git a/Engine/Source/Editor/BlueprintGraph/Private/K2Node_Event.cpp b/Engine/Source/Editor/BlueprintGraph/Private/K2Node_Event.cpp
--- a/Engine/Source/Editor/BlueprintGraph/Private/K2Node_Event.cpp
+++ b/Engine/Source/Editor/BlueprintGraph/Private/K2Node_Event.cpp
@@ -406,7 +406,7 @@ bool UK2Node_Event::CanPasteHere(const UEdGraph* TargetGraph) const
 			if(Blueprint->ParentClass->HasMetaData(*ExclusionListKeyName))
 			{
 				const FString ExcludedEventNameString = Blueprint->ParentClass->GetMetaData(*ExclusionListKeyName);
-				ExcludedEventNameString.ParseIntoArray(&ExcludedEventNames, TEXT(","), true);
+				ExcludedEventNameString.ParseIntoArray(ExcludedEventNames, TEXT(","), true);
 			}
 
 			// Gather all interfaces implemented by the Blueprint class
diff --git a/Engine/Source/Editor/ContentBrowser/Private/AssetViewSortManager.cpp b/Engine/Source/Editor/ContentBrowser/Private/AssetViewSortManager.cpp
--- a/Engine/Source/Editor/ContentBrowser/Private/AssetViewSortManager.cpp
+++ b/Engine/Source/Editor/ContentBrowser/Private/AssetViewSortManager.cpp
@@ -268,7 +268,7 @@ protected:
 		float Num1 = 1.f;
 		{
 			TArray<FString> Tokens1;
-			Value1->ParseIntoArray(&Tokens1, TEXT("x"), true);
+			Value1->ParseIntoArray(Tokens1, TEXT("x"), true);
 			for (auto TokenIt1 = Tokens1.CreateConstIterator(); TokenIt1; ++TokenIt1)
 			{
 				Num1 *= FCString::Atof(**TokenIt1);
@@ -277,7 +277,7 @@ protected:
 		float Num2 = 1.f;
 		{
 			TArray<FString> Tokens2;
-			Value2->ParseIntoArray(&Tokens2, TEXT("x"), true);
+			Value2->ParseIntoArray(Tokens2, TEXT("x"), true);
 			for (auto TokenIt2 = Tokens2.CreateConstIterator(); TokenIt2; ++TokenIt2)
 			{
 				Num2 *= FCString::Atof(**TokenIt2);
diff --git a/Engine/Source/Editor/ContentBrowser/Private/NativeClassHierarchy.cpp b/Engine/Source/Editor/ContentBrowser/Private/NativeClassHierarchy.cpp
--- a/Engine/Source/Editor/ContentBrowser/Private/NativeClassHierarchy.cpp
+++ b/Engine/Source/Editor/ContentBrowser/Private/NativeClassHierarchy.cpp
@@ -151,7 +151,7 @@ void FNativeClassHierarchy::GatherMatchingNodesForPaths(const TArray<FName>& InC
 			TSharedPtr<FNativeClassHierarchyNode> CurrentNode;
 
 			TArray<FString> ClassPathParts;
-			ClassPath.ToString().ParseIntoArray(&ClassPathParts, TEXT("/"), true);
+			ClassPath.ToString().ParseIntoArray(ClassPathParts, TEXT("/"), true);
 			for(const FString& ClassPathPart : ClassPathParts)
 			{
 				// Try and find the node associated with this part of the path...
@@ -286,7 +286,7 @@ void FNativeClassHierarchy::AddClass(UClass* InClass, const TSet<FName>& InGameM
 
 	// Split the class path and ensure we have nodes for each part
 	TArray<FString> HierarchyPathParts;
-	ClassHierarchyPath.ParseIntoArray(&HierarchyPathParts, TEXT("/"), true);
+	ClassHierarchyPath.ParseIntoArray(HierarchyPathParts, TEXT("/"), true);
 	TSharedPtr<FNativeClassHierarchyNode> CurrentNode = RootNode;
 	for(const FString& HierarchyPathPart : HierarchyPathParts)
 	{
@@ -311,7 +311,7 @@ void FNativeClassHierarchy::AddFolder(const FString& InClassPath)
 
 	// Split the class path and ensure we have nodes for each part
 	TArray<FString> ClassPathParts;
-	InClassPath.ParseIntoArray(&ClassPathParts, TEXT("/"), true);
+	InClassPath.ParseIntoArray(ClassPathParts, TEXT("/"), true);
 	TSharedPtr<FNativeClassHierarchyNode> CurrentNode;
 	for(const FString& ClassPathPart : ClassPathParts)
 	{
@@ -335,7 +335,7 @@ bool FNativeClassHierarchy::GetFileSystemPath(const FString& InClassPath, FStrin
 {
 	// Split the class path into its component parts
 	TArray<FString> ClassPathParts;
-	InClassPath.ParseIntoArray(&ClassPathParts, TEXT("/"), true);
+	InClassPath.ParseIntoArray(ClassPathParts, TEXT("/"), true);
 	
 	// We need to have at least two sections (a root, and a module name) to be able to resolve a file system path
 	if(ClassPathParts.Num() < 2)
diff --git a/Engine/Source/Editor/ContentBrowser/Private/SAssetView.cpp b/Engine/Source/Editor/ContentBrowser/Private/SAssetView.cpp
--- a/Engine/Source/Editor/ContentBrowser/Private/SAssetView.cpp
+++ b/Engine/Source/Editor/ContentBrowser/Private/SAssetView.cpp
@@ -1976,7 +1976,7 @@ void SAssetView::OnAssetRegistryPathAdded(const FString& Path)
 					const FString SubPath = Path.RightChop(SourcePath.Len());
 					
 					TArray<FString> SubPathItemList;
-					SubPath.ParseIntoArray(&SubPathItemList, TEXT("/"), /*InCullEmpty=*/true);
+					SubPath.ParseIntoArray(SubPathItemList, TEXT("/"), /*InCullEmpty=*/true);
 
 					if(SubPathItemList.Num() > 0)
 					{
diff --git a/Engine/Source/Editor/ContentBrowser/Private/SCollectionView.cpp b/Engine/Source/Editor/ContentBrowser/Private/SCollectionView.cpp
--- a/Engine/Source/Editor/ContentBrowser/Private/SCollectionView.cpp
+++ b/Engine/Source/Editor/ContentBrowser/Private/SCollectionView.cpp
@@ -277,7 +277,7 @@ void SCollectionView::LoadSettings(const FString& IniFilename, const FString& In
 	if ( GConfig->GetString(*IniSection, *(SettingsString + TEXT(".SelectedCollections")), SelectedCollectionsString, IniFilename) )
 	{
 		TArray<FString> NewSelectedCollectionStrings;
-		SelectedCollectionsString.ParseIntoArray(&NewSelectedCollectionStrings, TEXT(","), /*bCullEmpty*/true);
+		SelectedCollectionsString.ParseIntoArray(NewSelectedCollectionStrings, TEXT(","), /*bCullEmpty*/true);
 
 		for ( auto CollectionIt = NewSelectedCollectionStrings.CreateConstIterator(); CollectionIt; ++CollectionIt )
 		{
diff --git a/Engine/Source/Editor/ContentBrowser/Private/SContentBrowser.cpp b/Engine/Source/Editor/ContentBrowser/Private/SContentBrowser.cpp
--- a/Engine/Source/Editor/ContentBrowser/Private/SContentBrowser.cpp
+++ b/Engine/Source/Editor/ContentBrowser/Private/SContentBrowser.cpp
@@ -1901,7 +1901,7 @@ void SContentBrowser::UpdatePath()
 	if ( SourcesData.PackagePaths.Num() > 0 )
 	{
 		TArray<FString> Crumbs;
-		SourcesData.PackagePaths[0].ToString().ParseIntoArray(&Crumbs, TEXT("/"), true);
+		SourcesData.PackagePaths[0].ToString().ParseIntoArray(Crumbs, TEXT("/"), true);
 
 		FString CrumbPath = TEXT("/");
 		for ( auto CrumbIt = Crumbs.CreateConstIterator(); CrumbIt; ++CrumbIt )
diff --git a/Engine/Source/Editor/ContentBrowser/Private/SFilterList.cpp b/Engine/Source/Editor/ContentBrowser/Private/SFilterList.cpp
--- a/Engine/Source/Editor/ContentBrowser/Private/SFilterList.cpp
+++ b/Engine/Source/Editor/ContentBrowser/Private/SFilterList.cpp
@@ -653,8 +653,8 @@ void SFilterList::LoadSettings(const FString& IniFilename, const FString& IniSec
 		// Parse comma delimited strings into arrays
 		TArray<FString> TypeFilterNames;
 		TArray<FString> EnabledTypeFilterNames;
-		ActiveTypeFilterString.ParseIntoArray(&TypeFilterNames, TEXT(","), /*bCullEmpty=*/true);
-		EnabledTypeFilterString.ParseIntoArray(&EnabledTypeFilterNames, TEXT(","), /*bCullEmpty=*/true);
+		ActiveTypeFilterString.ParseIntoArray(TypeFilterNames, TEXT(","), /*bCullEmpty=*/true);
+		EnabledTypeFilterString.ParseIntoArray(EnabledTypeFilterNames, TEXT(","), /*bCullEmpty=*/true);
 
 		// Get the list of all asset type actions
 		FAssetToolsModule& AssetToolsModule = FModuleManager::LoadModuleChecked<FAssetToolsModule>("AssetTools");
@@ -691,8 +691,8 @@ void SFilterList::LoadSettings(const FString& IniFilename, const FString& IniSec
 		// Parse comma delimited strings into arrays
 		TArray<FString> FrontendFilterNames;
 		TArray<FString> EnabledFrontendFilterNames;
-		ActiveFrontendFilterString.ParseIntoArray(&FrontendFilterNames, TEXT(","), /*bCullEmpty=*/true);
-		EnabledFrontendFilterString.ParseIntoArray(&EnabledFrontendFilterNames, TEXT(","), /*bCullEmpty=*/true);
+		ActiveFrontendFilterString.ParseIntoArray(FrontendFilterNames, TEXT(","), /*bCullEmpty=*/true);
+		EnabledFrontendFilterString.ParseIntoArray(EnabledFrontendFilterNames, TEXT(","), /*bCullEmpty=*/true);
 
 		// For each FrontendFilter, add any that were active and enable any that were previously enabled
 		for ( auto FrontendFilterIt = AllFrontendFilters.CreateConstIterator(); FrontendFilterIt; ++FrontendFilterIt )
diff --git a/Engine/Source/Editor/ContentBrowser/Private/SPathView.cpp b/Engine/Source/Editor/ContentBrowser/Private/SPathView.cpp
--- a/Engine/Source/Editor/ContentBrowser/Private/SPathView.cpp
+++ b/Engine/Source/Editor/ContentBrowser/Private/SPathView.cpp
@@ -178,7 +178,7 @@ void SPathView::SetSelectedPaths(const TArray<FString>& Paths)
 		const FString& Path = Paths[PathIdx];
 
 		TArray<FString> PathItemList;
-		Path.ParseIntoArray(&PathItemList, TEXT("/"), /*InCullEmpty=*/true);
+		Path.ParseIntoArray(PathItemList, TEXT("/"), /*InCullEmpty=*/true);
 
 		if ( PathItemList.Num() )
 		{
@@ -283,7 +283,7 @@ TSharedPtr<FTreeItem> SPathView::AddPath(const FString& Path, bool bUserNamed)
 	}
 
 	TArray<FString> PathItemList;
-	Path.ParseIntoArray(&PathItemList, TEXT("/"), /*InCullEmpty=*/true);
+	Path.ParseIntoArray(PathItemList, TEXT("/"), /*InCullEmpty=*/true);
 
 	if ( PathItemList.Num() )
 	{
@@ -584,7 +584,7 @@ void SPathView::LoadSettings(const FString& IniFilename, const FString& IniSecti
 	TArray<FString> NewSelectedPaths;
 	if ( GConfig->GetString(*IniSection, *(SettingsString + TEXT(".SelectedPaths")), SelectedPathsString, IniFilename) )
 	{
-		SelectedPathsString.ParseIntoArray(&NewSelectedPaths, TEXT(","), /*bCullEmpty*/true);
+		SelectedPathsString.ParseIntoArray(NewSelectedPaths, TEXT(","), /*bCullEmpty*/true);
 	}
 
 	if ( NewSelectedPaths.Num() > 0 )
@@ -682,7 +682,7 @@ bool SPathView::ExplicitlyAddPathToSelection(const FString& Path)
 	}
 
 	TArray<FString> PathItemList;
-	Path.ParseIntoArray(&PathItemList, TEXT("/"), /*InCullEmpty=*/true);
+	Path.ParseIntoArray(PathItemList, TEXT("/"), /*InCullEmpty=*/true);
 
 	if ( PathItemList.Num() )
 	{
diff --git a/Engine/Source/Editor/DetailCustomizations/Private/FbxImportUIDetails.cpp b/Engine/Source/Editor/DetailCustomizations/Private/FbxImportUIDetails.cpp
--- a/Engine/Source/Editor/DetailCustomizations/Private/FbxImportUIDetails.cpp
+++ b/Engine/Source/Editor/DetailCustomizations/Private/FbxImportUIDetails.cpp
@@ -305,7 +305,7 @@ void FFbxImportUIDetails::CollectChildPropertiesRecursive(TSharedPtr<IPropertyHa
 bool FFbxImportUIDetails::IsImportTypeMetaDataValid(EFBXImportType& ImportType, FString& MetaData)
 {
 	TArray<FString> Types;
-	MetaData.ParseIntoArray(&Types, TEXT("|"), 1);
+	MetaData.ParseIntoArray(Types, TEXT("|"), 1);
 	switch(ImportType)
 	{
 		case FBXIT_StaticMesh:
diff --git a/Engine/Source/Editor/DetailCustomizations/Private/StringAssetReferenceCustomization.cpp b/Engine/Source/Editor/DetailCustomizations/Private/StringAssetReferenceCustomization.cpp
--- a/Engine/Source/Editor/DetailCustomizations/Private/StringAssetReferenceCustomization.cpp
+++ b/Engine/Source/Editor/DetailCustomizations/Private/StringAssetReferenceCustomization.cpp
@@ -12,7 +12,7 @@ void FStringAssetReferenceCustomization::CustomizeHeader( TSharedRef<IPropertyHa
 	if( !ClassFilterString.IsEmpty() )
 	{
 		TArray<FString> CustomClassFilterNames;
-		ClassFilterString.ParseIntoArray(&CustomClassFilterNames, TEXT(","), true);
+		ClassFilterString.ParseIntoArray(CustomClassFilterNames, TEXT(","), true);
 
 		for(auto It = CustomClassFilterNames.CreateConstIterator(); It; ++It)
 		{
diff --git a/Engine/Source/Editor/Documentation/Private/UDNParser.cpp b/Engine/Source/Editor/Documentation/Private/UDNParser.cpp
--- a/Engine/Source/Editor/Documentation/Private/UDNParser.cpp
+++ b/Engine/Source/Editor/Documentation/Private/UDNParser.cpp
@@ -1263,7 +1263,7 @@ bool FUDNParser::ParseCodeLink(FString &InternalLink)
 	int32 Col = 0;
 
 	TArray<FString> Tokens;
-	InternalLink.ParseIntoArray(&Tokens, TEXT(","), 0);
+	InternalLink.ParseIntoArray(Tokens, TEXT(","), 0);
 	int32 TokenStringsCount = Tokens.Num();
 	if (TokenStringsCount > 0)
 	{
@@ -1317,7 +1317,7 @@ bool FUDNParser::ParseCodeLink(FString &InternalLink)
 bool FUDNParser::ParseAssetLink(FString &InternalLink)
 {
 	TArray<FString> Token;
-	InternalLink.ParseIntoArray(&Token, TEXT(","), 0);
+	InternalLink.ParseIntoArray(Token, TEXT(","), 0);
 	
 	if (Token.Num() >= 2)
 	{
diff --git a/Engine/Source/Editor/FoliageEdit/Private/FoliageEdMode.cpp b/Engine/Source/Editor/FoliageEdit/Private/FoliageEdMode.cpp
--- a/Engine/Source/Editor/FoliageEdit/Private/FoliageEdMode.cpp
+++ b/Engine/Source/Editor/FoliageEdit/Private/FoliageEdMode.cpp
@@ -2637,7 +2637,7 @@ void FFoliageUISettings::Load()
 	if (GConfig->GetString(TEXT("FoliageEdit"), TEXT("WindowPosition"), WindowPositionString, GEditorUserSettingsIni))
 	{
 		TArray<FString> PositionValues;
-		if (WindowPositionString.ParseIntoArray(&PositionValues, TEXT(","), true) == 4)
+		if (WindowPositionString.ParseIntoArray(PositionValues, TEXT(","), true) == 4)
 		{
 			WindowX = FCString::Atoi(*PositionValues[0]);
 			WindowY = FCString::Atoi(*PositionValues[1]);
diff --git a/Engine/Source/Editor/GameProjectGeneration/Private/GameProjectUtils.cpp b/Engine/Source/Editor/GameProjectGeneration/Private/GameProjectUtils.cpp
--- a/Engine/Source/Editor/GameProjectGeneration/Private/GameProjectUtils.cpp
+++ b/Engine/Source/Editor/GameProjectGeneration/Private/GameProjectUtils.cpp
@@ -3097,7 +3097,7 @@ void GameProjectUtils::HarvestCursorSyncLocation( FString& FinalOutput, FString&
 
 	// Determine the cursor focus location if this file will by synced after creation
 	TArray<FString> Lines;
-	FinalOutput.ParseIntoArray( &Lines, TEXT( "\n" ), false );
+	FinalOutput.ParseIntoArray( Lines, TEXT( "\n" ), false );
 	for( int32 LineIdx = 0; LineIdx < Lines.Num(); ++LineIdx )
 	{
 		const FString& Line = Lines[ LineIdx ];
diff --git a/Engine/Source/Editor/GameplayAbilitiesEditor/Private/K2Node_LatentAbilityCall.cpp b/Engine/Source/Editor/GameplayAbilitiesEditor/Private/K2Node_LatentAbilityCall.cpp
--- a/Engine/Source/Editor/GameplayAbilitiesEditor/Private/K2Node_LatentAbilityCall.cpp
+++ b/Engine/Source/Editor/GameplayAbilitiesEditor/Private/K2Node_LatentAbilityCall.cpp
@@ -234,7 +234,7 @@ void UK2Node_LatentAbilityCall::CreatePinsForClass(UClass* InClass)
 	
 		if (!IgnorePropertyListStr.IsEmpty())
 		{
-			IgnorePropertyListStr.ParseIntoArray(&IgnorePropertyList, TEXT(","), true);
+			IgnorePropertyListStr.ParseIntoArray(IgnorePropertyList, TEXT(","), true);
 		}
 	}
 
diff --git a/Engine/Source/Editor/GraphEditor/Private/KismetPins/SGraphPinObject.cpp b/Engine/Source/Editor/GraphEditor/Private/KismetPins/SGraphPinObject.cpp
--- a/Engine/Source/Editor/GraphEditor/Private/KismetPins/SGraphPinObject.cpp
+++ b/Engine/Source/Editor/GraphEditor/Private/KismetPins/SGraphPinObject.cpp
@@ -194,7 +194,7 @@ TSharedRef<SWidget> SGraphPinObject::GenerateAssetPicker()
 
 		// Parse and add the classes from the metadata
 		TArray<FString> CustomClassFilterNames;
-		ClassFilterString.ParseIntoArray(&CustomClassFilterNames, TEXT(","), true);
+		ClassFilterString.ParseIntoArray(CustomClassFilterNames, TEXT(","), true);
 		for(auto It = CustomClassFilterNames.CreateConstIterator(); It; ++It)
 		{
 			AssetPickerConfig.Filter.ClassNames.Add(FName(**It));
diff --git a/Engine/Source/Editor/GraphEditor/Private/KismetPins/SGraphPinVector.cpp b/Engine/Source/Editor/GraphEditor/Private/KismetPins/SGraphPinVector.cpp
--- a/Engine/Source/Editor/GraphEditor/Private/KismetPins/SGraphPinVector.cpp
+++ b/Engine/Source/Editor/GraphEditor/Private/KismetPins/SGraphPinVector.cpp
@@ -193,7 +193,7 @@ FString SGraphPinVector::GetValue( ETextBoxIndex Index ) const
 	//Parse string to split its contents separated by ','
 	DefaultString.Trim();
 	DefaultString.TrimTrailing();
-	DefaultString.ParseIntoArray(&ResultString, TEXT(","), true);
+	DefaultString.ParseIntoArray(ResultString, TEXT(","), true);
 
 	if(Index < ResultString.Num())
 	{
diff --git a/Engine/Source/Editor/GraphEditor/Private/SGraphActionMenu.cpp b/Engine/Source/Editor/GraphEditor/Private/SGraphActionMenu.cpp
--- a/Engine/Source/Editor/GraphEditor/Private/SGraphActionMenu.cpp
+++ b/Engine/Source/Editor/GraphEditor/Private/SGraphActionMenu.cpp
@@ -637,7 +637,7 @@ void SGraphActionMenu::GenerateFilteredItems(bool bPreserveExpansion)
 
 	// Tokenize the search box text into a set of terms; all of them must be present to pass the filter
 	TArray<FString> FilterTerms;
-	TrimmedFilterString.ParseIntoArray(/*out*/ &FilterTerms, TEXT(" "), true);
+	TrimmedFilterString.ParseIntoArray(FilterTerms, TEXT(" "), true);
 	
 	// Generate a list of sanitized versions of the strings
 	TArray<FString> SanitizedFilterTerms;
@@ -772,22 +772,22 @@ int32 SGraphActionMenu::GetActionFilteredWeight( const FGraphActionListBuilderBa
 		SearchText = SearchText.Replace( TEXT( " " ), TEXT( "" ) );
 
 		// First the keywords
-		InCurrentAction.Actions[Action]->Keywords.ParseIntoArray( &EachEntry.Array, TEXT(" "), true );
+		InCurrentAction.Actions[Action]->Keywords.ParseIntoArray( EachEntry.Array, TEXT(" "), true );
 		EachEntry.Weight = 10;
 		WeightedArrayList.Add( EachEntry );
 
 		// The description
-		InCurrentAction.Actions[Action]->MenuDescription.ToString().ParseIntoArray( &EachEntry.Array, TEXT(" "), true );
+		InCurrentAction.Actions[Action]->MenuDescription.ToString().ParseIntoArray( EachEntry.Array, TEXT(" "), true );
 		EachEntry.Weight = DescriptionWeight;
 		WeightedArrayList.Add( EachEntry );
 
 		// The node search title weight
-		InCurrentAction.Actions[Action]->GetSearchTitle().ToString().ParseIntoArray( &EachEntry.Array, TEXT(" "), true );
+		InCurrentAction.Actions[Action]->GetSearchTitle().ToString().ParseIntoArray( EachEntry.Array, TEXT(" "), true );
 		EachEntry.Weight = NodeTitleWeight;
 		WeightedArrayList.Add( EachEntry );
 
 		// The category
-		InCurrentAction.Actions[Action]->Category.ParseIntoArray( &EachEntry.Array, TEXT(" "), true );
+		InCurrentAction.Actions[Action]->Category.ParseIntoArray( EachEntry.Array, TEXT(" "), true );
 		EachEntry.Weight = CategoryWeight;
 		WeightedArrayList.Add( EachEntry );
 
diff --git a/Engine/Source/Editor/GraphEditor/Private/SGraphNode.cpp b/Engine/Source/Editor/GraphEditor/Private/SGraphNode.cpp
--- a/Engine/Source/Editor/GraphEditor/Private/SGraphNode.cpp
+++ b/Engine/Source/Editor/GraphEditor/Private/SGraphNode.cpp
@@ -69,7 +69,7 @@ void SNodeTitle::RebuildWidget()
 	// Break the title into lines
 	TArray<FString> Lines;
 	const FString CachedTitleString = CachedTitle.ToString().Replace(TEXT("\r"), TEXT(""));
-	CachedTitleString.ParseIntoArray(&Lines, TEXT("\n"), false);
+	CachedTitleString.ParseIntoArray(Lines, TEXT("\n"), false);
 
 	if (Lines.Num())
 	{
diff --git a/Engine/Source/Editor/InputBindingEditor/Private/Widgets/SInputBindingEditorPanel.cpp b/Engine/Source/Editor/InputBindingEditor/Private/Widgets/SInputBindingEditorPanel.cpp
--- a/Engine/Source/Editor/InputBindingEditor/Private/Widgets/SInputBindingEditorPanel.cpp
+++ b/Engine/Source/Editor/InputBindingEditor/Private/Widgets/SInputBindingEditorPanel.cpp
@@ -123,7 +123,7 @@ void SInputBindingEditorPanel::OnSearchChanged( const FText& NewSearch )
 	// Remove whitespace from the front and back of the string
 	ParseString.Trim();
 	ParseString.TrimTrailing();
-	ParseString.ParseIntoArray(&FilterStrings, TEXT(" "), true);
+	ParseString.ParseIntoArray(FilterStrings, TEXT(" "), true);
 
 	FilterVisibleContextList();
 }
diff --git a/Engine/Source/Editor/IntroTutorials/Private/STutorialsBrowser.cpp b/Engine/Source/Editor/IntroTutorials/Private/STutorialsBrowser.cpp
--- a/Engine/Source/Editor/IntroTutorials/Private/STutorialsBrowser.cpp
+++ b/Engine/Source/Editor/IntroTutorials/Private/STutorialsBrowser.cpp
@@ -638,7 +638,7 @@ TSharedPtr<FTutorialListEntry_Category> STutorialsBrowser::RebuildCategories()
 		// We're expecting the category string to be in the "A.B.C" format.  We'll split up the string here and form
 		// a proper hierarchy in the UI
 		TArray< FString > SplitCategories;
-		CategoryPath.ParseIntoArray( &SplitCategories, TEXT( "." ), true /* bCullEmpty */ );
+		CategoryPath.ParseIntoArray( SplitCategories, TEXT( "." ), true /* bCullEmpty */ );
 
 		FString CurrentCategoryPath;
 
@@ -736,7 +736,7 @@ void STutorialsBrowser::RebuildTutorials(TSharedPtr<FTutorialListEntry_Category>
 		// We're expecting the category string to be in the "A.B.C" format.  We'll split up the string here and form
 		// a proper hierarchy in the UI
 		TArray< FString > SplitCategories;
-		CategoryPath.ParseIntoArray( &SplitCategories, TEXT( "." ), true /* bCullEmpty */ );
+		CategoryPath.ParseIntoArray( SplitCategories, TEXT( "." ), true /* bCullEmpty */ );
 
 		FString CurrentCategoryPath;
 
diff --git a/Engine/Source/Editor/IntroTutorials/Private/TutorialText.cpp b/Engine/Source/Editor/IntroTutorials/Private/TutorialText.cpp
--- a/Engine/Source/Editor/IntroTutorials/Private/TutorialText.cpp
+++ b/Engine/Source/Editor/IntroTutorials/Private/TutorialText.cpp
@@ -135,7 +135,7 @@ static void ParseCodeLink(const FString &InternalLink)
 	int32 Col = 0;
 
 	TArray<FString> Tokens;
-	InternalLink.ParseIntoArray(&Tokens, TEXT(","), 0);
+	InternalLink.ParseIntoArray(Tokens, TEXT(","), 0);
 	int32 TokenStringsCount = Tokens.Num();
 	if (TokenStringsCount > 0)
 	{
diff --git a/Engine/Source/Editor/Kismet/Private/BlueprintActionMenuUtils.cpp b/Engine/Source/Editor/Kismet/Private/BlueprintActionMenuUtils.cpp
--- a/Engine/Source/Editor/Kismet/Private/BlueprintActionMenuUtils.cpp
+++ b/Engine/Source/Editor/Kismet/Private/BlueprintActionMenuUtils.cpp
@@ -178,7 +178,7 @@ static bool BlueprintActionMenuUtilsImpl::IsUnexposedMemberAction(FBlueprintActi
 				}
 
 				TArray<FString> PropertyExposedCategories;
-				ExposedCategoryMetadata.ParseIntoArray(&PropertyExposedCategories, TEXT(","), true);
+				ExposedCategoryMetadata.ParseIntoArray(PropertyExposedCategories, TEXT(","), true);
 				AllExposedCategories.Append(PropertyExposedCategories);
 			}
 		}
diff --git a/Engine/Source/Editor/Kismet/Private/BlueprintEditorCommands.cpp b/Engine/Source/Editor/Kismet/Private/BlueprintEditorCommands.cpp
--- a/Engine/Source/Editor/Kismet/Private/BlueprintEditorCommands.cpp
+++ b/Engine/Source/Editor/Kismet/Private/BlueprintEditorCommands.cpp
@@ -95,7 +95,7 @@ namespace NodeSpawnInfoHelpers
 		if( InBlueprint->ParentClass->HasMetaData(*ExclusionListKeyName) )
 		{
 			const FString ExcludedEventNameString = InBlueprint->ParentClass->GetMetaData(*ExclusionListKeyName);
-			ExcludedEventNameString.ParseIntoArray(&ExcludedEventNames, TEXT(","), true);
+			ExcludedEventNameString.ParseIntoArray(ExcludedEventNames, TEXT(","), true);
 		}
 
 		const UEdGraphSchema_K2* K2Schema = GetDefault<UEdGraphSchema_K2>();
diff --git a/Engine/Source/Editor/Kismet/Private/FindInBlueprints.cpp b/Engine/Source/Editor/Kismet/Private/FindInBlueprints.cpp
--- a/Engine/Source/Editor/Kismet/Private/FindInBlueprints.cpp
+++ b/Engine/Source/Editor/Kismet/Private/FindInBlueprints.cpp
@@ -1170,7 +1170,7 @@ ECheckBoxState SFindInBlueprints::OnGetFindModeChecked() const
 void SFindInBlueprints::InitiateSearch()
 {
 	TArray<FString> Tokens;
-	if(SearchValue.Contains("\"") && SearchValue.ParseIntoArray(&Tokens, TEXT("\""), true)>0)
+	if(SearchValue.Contains("\"") && SearchValue.ParseIntoArray(Tokens, TEXT("\""), true)>0)
 	{
 		for( auto &TokenIt : Tokens )
 		{
@@ -1194,7 +1194,7 @@ void SFindInBlueprints::InitiateSearch()
 	else
 	{
 		// unquoted search equivalent to a match-any-of search
-		SearchValue.ParseIntoArray(&Tokens, TEXT(" "), true);
+		SearchValue.ParseIntoArray(Tokens, TEXT(" "), true);
 	}
 
 	if(ItemsFound.Num())
diff --git a/Engine/Source/Editor/Kismet/Private/SFilterableObjectList.cpp b/Engine/Source/Editor/Kismet/Private/SFilterableObjectList.cpp
--- a/Engine/Source/Editor/Kismet/Private/SFilterableObjectList.cpp
+++ b/Engine/Source/Editor/Kismet/Private/SFilterableObjectList.cpp
@@ -153,7 +153,7 @@ void SFilterableObjectList::RefilterObjectList()
 	TArray<FString> FilterTerms;
 	if (FilterTextBoxWidget.IsValid())
 	{
-		FilterTextBoxWidget->GetText().ToString().ParseIntoArray(/*out*/ &FilterTerms, TEXT(" "), true);
+		FilterTextBoxWidget->GetText().ToString().ParseIntoArray(FilterTerms, TEXT(" "), true);
 	}
 
 	if (FilterTerms.Num())
diff --git a/Engine/Source/Editor/KismetCompiler/Private/KismetCompilerMisc.cpp b/Engine/Source/Editor/KismetCompiler/Private/KismetCompilerMisc.cpp
--- a/Engine/Source/Editor/KismetCompiler/Private/KismetCompilerMisc.cpp
+++ b/Engine/Source/Editor/KismetCompiler/Private/KismetCompilerMisc.cpp
@@ -46,12 +46,12 @@ bool FKismetCompilerUtilities::IsTypeCompatibleWithProperty(UEdGraphPin* SourceP
 				// Check for the magic ArrayParm property, which always matches array types
 				FString ArrayPointerMetaData = OwningFunction->GetMetaData(TEXT("ArrayParm"));
 				TArray<FString> ArrayPinComboNames;
-				ArrayPointerMetaData.ParseIntoArray(&ArrayPinComboNames, TEXT(","), true);
+				ArrayPointerMetaData.ParseIntoArray(ArrayPinComboNames, TEXT(","), true);
 
 				for(auto Iter = ArrayPinComboNames.CreateConstIterator(); Iter; ++Iter)
 				{
 					TArray<FString> ArrayPinNames;
-					Iter->ParseIntoArray(&ArrayPinNames, TEXT("|"), true);
+					Iter->ParseIntoArray(ArrayPinNames, TEXT("|"), true);
 
 					if( ArrayPinNames[0] == SourcePin->PinName )
 					{
@@ -80,7 +80,7 @@ bool FKismetCompilerUtilities::IsTypeCompatibleWithProperty(UEdGraphPin* SourceP
 		// Check to see if this param is type dependent on an array parameter
 		const FString DependentParams = OwningFunction->GetMetaData(TEXT("ArrayTypeDependentParams"));
 		TArray<FString>	DependentParamNames;
-		DependentParams.ParseIntoArray(&DependentParamNames, TEXT(","), true);
+		DependentParams.ParseIntoArray(DependentParamNames, TEXT(","), true);
 		if (DependentParamNames.Find(SourcePin->PinName) != INDEX_NONE)
 		{
 			//@todo:  This assumes that the wildcard coersion has done its job...I'd feel better if there was some easier way of accessing the target array type
diff --git a/Engine/Source/Editor/LandscapeEditor/Private/SLandscapeEditor.cpp b/Engine/Source/Editor/LandscapeEditor/Private/SLandscapeEditor.cpp
--- a/Engine/Source/Editor/LandscapeEditor/Private/SLandscapeEditor.cpp
+++ b/Engine/Source/Editor/LandscapeEditor/Private/SLandscapeEditor.cpp
@@ -290,7 +290,7 @@ bool SLandscapeEditor::GetIsPropertyVisible(const FPropertyAndParent& PropertyAn
 			const FName CurrentToolName = LandscapeEdMode->CurrentTool->GetToolName();
 
 			TArray<FString> ShowForTools;
-			Property.GetMetaData("ShowForTools").ParseIntoArray(&ShowForTools, TEXT(","), true);
+			Property.GetMetaData("ShowForTools").ParseIntoArray(ShowForTools, TEXT(","), true);
 			if (!ShowForTools.Contains(CurrentToolName.ToString()))
 			{
 				return false;
@@ -302,7 +302,7 @@ bool SLandscapeEditor::GetIsPropertyVisible(const FPropertyAndParent& PropertyAn
 			// const FName CurrentBrushName = LandscapeEdMode->CurrentBrush->GetBrushName();
 
 			TArray<FString> ShowForBrushes;
-			Property.GetMetaData("ShowForBrushes").ParseIntoArray(&ShowForBrushes, TEXT(","), true);
+			Property.GetMetaData("ShowForBrushes").ParseIntoArray(ShowForBrushes, TEXT(","), true);
 			if (!ShowForBrushes.Contains(CurrentBrushSetName.ToString()))
 				//&& !ShowForBrushes.Contains(CurrentBrushName.ToString())
 			{
@@ -314,7 +314,7 @@ bool SLandscapeEditor::GetIsPropertyVisible(const FPropertyAndParent& PropertyAn
 			static const TCHAR* TargetTypeNames[] = { TEXT("Heightmap"), TEXT("Weightmap"), TEXT("Visibility") };
 
 			TArray<FString> ShowForTargetTypes;
-			Property.GetMetaData("ShowForTargetTypes").ParseIntoArray(&ShowForTargetTypes, TEXT(","), true);
+			Property.GetMetaData("ShowForTargetTypes").ParseIntoArray(ShowForTargetTypes, TEXT(","), true);
 
 			const ELandscapeToolTargetType::Type CurrentTargetType = LandscapeEdMode->CurrentToolTarget.TargetType;
 			if (CurrentTargetType == ELandscapeToolTargetType::Invalid ||
diff --git a/Engine/Source/Editor/LevelEditor/Private/SLevelEditorBuildAndSubmit.cpp b/Engine/Source/Editor/LevelEditor/Private/SLevelEditorBuildAndSubmit.cpp
--- a/Engine/Source/Editor/LevelEditor/Private/SLevelEditorBuildAndSubmit.cpp
+++ b/Engine/Source/Editor/LevelEditor/Private/SLevelEditorBuildAndSubmit.cpp
@@ -408,7 +408,7 @@ FReply SLevelEditorBuildAndSubmit::OnBuildClicked()
 		BuildAndSubmitErrors.NewPage(LOCTEXT("BuildAndSubmitErrorsNewPage", "Build and Submit"));
 
 		TArray<FString> Errors;
-		ErrorMessage.ToString().Replace(LINE_TERMINATOR, TEXT("\n")).ParseIntoArray(&Errors, TEXT("\n"), true);
+		ErrorMessage.ToString().Replace(LINE_TERMINATOR, TEXT("\n")).ParseIntoArray(Errors, TEXT("\n"), true);
 
 		for (int32 ErrorIdx = 0; ErrorIdx < Errors.Num(); ErrorIdx++)
 		{
diff --git a/Engine/Source/Editor/MaterialEditor/Private/FindInMaterial.cpp b/Engine/Source/Editor/MaterialEditor/Private/FindInMaterial.cpp
--- a/Engine/Source/Editor/MaterialEditor/Private/FindInMaterial.cpp
+++ b/Engine/Source/Editor/MaterialEditor/Private/FindInMaterial.cpp
@@ -196,7 +196,7 @@ void SFindInMaterial::OnSearchTextCommitted(const FText& Text, ETextCommit::Type
 void SFindInMaterial::InitiateSearch()
 {
 	TArray<FString> Tokens;
-	if (SearchValue.Contains("\"") && SearchValue.ParseIntoArray(&Tokens, TEXT("\""), true) > 0)
+	if (SearchValue.Contains("\"") && SearchValue.ParseIntoArray(Tokens, TEXT("\""), true) > 0)
 	{
 		for (auto &TokenIt : Tokens)
 		{
@@ -220,7 +220,7 @@ void SFindInMaterial::InitiateSearch()
 	else
 	{
 		// unquoted search equivalent to a match-any-of search
-		SearchValue.ParseIntoArray(&Tokens, TEXT(" "), true);
+		SearchValue.ParseIntoArray(Tokens, TEXT(" "), true);
 	}
 
 	for (auto It(ItemsFound.CreateIterator()); It; ++It)
diff --git a/Engine/Source/Editor/Matinee/Private/MatineeRecordMovie.cpp b/Engine/Source/Editor/Matinee/Private/MatineeRecordMovie.cpp
--- a/Engine/Source/Editor/Matinee/Private/MatineeRecordMovie.cpp
+++ b/Engine/Source/Editor/Matinee/Private/MatineeRecordMovie.cpp
@@ -163,7 +163,7 @@ FString SMatineeRecordMovie::CustomResolutionStr( TEXT("Custom") );
 static void ParseResolutionString( const FString& ResStr, uint32 &OutWidth, uint32& OutHeight )
 {
 	TArray<FString> Res;
-	if( ResStr.ParseIntoArray( &Res, TEXT("x"), 0 ) == 2 )
+	if( ResStr.ParseIntoArray( Res, TEXT("x"), 0 ) == 2 )
 	{
 		OutWidth = FCString::Atof(Res[0].GetCharArray().GetData());
 		OutHeight = FCString::Atof(Res[1].GetCharArray().GetData());
diff --git a/Engine/Source/Editor/Persona/Private/SSkeletonSmartNameManager.cpp b/Engine/Source/Editor/Persona/Private/SSkeletonSmartNameManager.cpp
--- a/Engine/Source/Editor/Persona/Private/SSkeletonSmartNameManager.cpp
+++ b/Engine/Source/Editor/Persona/Private/SSkeletonSmartNameManager.cpp
@@ -405,7 +405,7 @@ void SCurveNameManager::OnDeleteNameClicked()
 			}
 
 			TArray<FString> ParsedStringUids;
-			CurveDataCopy.ParseIntoArray(&ParsedStringUids, *USkeleton::CurveTagDelimiter, true);
+			CurveDataCopy.ParseIntoArray(ParsedStringUids, *USkeleton::CurveTagDelimiter, true);
 
 			for(const FString& UidString : ParsedStringUids)
 			{
diff --git a/Engine/Source/Editor/PropertyEditor/Private/DetailLayoutBuilderImpl.cpp b/Engine/Source/Editor/PropertyEditor/Private/DetailLayoutBuilderImpl.cpp
--- a/Engine/Source/Editor/PropertyEditor/Private/DetailLayoutBuilderImpl.cpp
+++ b/Engine/Source/Editor/PropertyEditor/Private/DetailLayoutBuilderImpl.cpp
@@ -295,7 +295,7 @@ TSharedPtr<FPropertyNode> FDetailLayoutBuilderImpl::GetPropertyNodeInternal( con
 {
 	FName PropertyName;
 	TArray<FString> PathList;
-	PropertyPath.ToString().ParseIntoArray( &PathList, TEXT("."), true );
+	PropertyPath.ToString().ParseIntoArray( PathList, TEXT("."), true );
 
 	if( PathList.Num() == 1 )
 	{
diff --git a/Engine/Source/Editor/PropertyEditor/Private/ObjectPropertyNode.cpp b/Engine/Source/Editor/PropertyEditor/Private/ObjectPropertyNode.cpp
--- a/Engine/Source/Editor/PropertyEditor/Private/ObjectPropertyNode.cpp
+++ b/Engine/Source/Editor/PropertyEditor/Private/ObjectPropertyNode.cpp
@@ -369,7 +369,7 @@ void FObjectPropertyNode::InternalInitChildNodes( FName SinglePropertyName )
 		{
 			// Figure out the nesting level for this category
 			TArray< FString > FullCategoryPathStrings;
-			FullCategoryPath.ToString().ParseIntoArray( &FullCategoryPathStrings, *CategoryDelimiterString, true );
+			FullCategoryPath.ToString().ParseIntoArray( FullCategoryPathStrings, *CategoryDelimiterString, true );
 
 			TSharedPtr<FPropertyNode> ParentLevelNode = SharedThis(this);
 			FString CurCategoryPathString;
diff --git a/Engine/Source/Editor/PropertyEditor/Private/Presentation/PropertyTable/PropertyTable.cpp b/Engine/Source/Editor/PropertyEditor/Private/Presentation/PropertyTable/PropertyTable.cpp
--- a/Engine/Source/Editor/PropertyEditor/Private/Presentation/PropertyTable/PropertyTable.cpp
+++ b/Engine/Source/Editor/PropertyEditor/Private/Presentation/PropertyTable/PropertyTable.cpp
@@ -412,11 +412,11 @@ void FPropertyTable::PasteTextAtCell( const FString& Text, const TSharedRef< IPr
 
 	// Parse into row strings
 	TArray<FString> RowStrings;
-	Text.ParseIntoArray(&RowStrings, TEXT("\n"), true);
+	Text.ParseIntoArray(RowStrings, TEXT("\n"), true);
 
 	// Parse row strings into individual cell strings
 	TArray<FString> CellStrings;
-	RowStrings[CurrentRowIdx++].ParseIntoArray(&CellStrings, TEXT("\t"), false);
+	RowStrings[CurrentRowIdx++].ParseIntoArray(CellStrings, TEXT("\t"), false);
 
 	// Get the maximum paste operations before displaying the slow task
 	int32 NumPasteOperationsBeforeWarning = GEditor->AccessEditorUserSettings().PropertyMatrix_NumberOfPasteOperationsBeforeWarning;
@@ -459,7 +459,7 @@ void FPropertyTable::PasteTextAtCell( const FString& Text, const TSharedRef< IPr
 					// Prepare data to operate on next row
 					CurrentColumnIdx = 0;
 					FirstCellInRow = TargetCell;
-					RowStrings[CurrentRowIdx++].ParseIntoArray(&CellStrings, TEXT("\t"), false);
+					RowStrings[CurrentRowIdx++].ParseIntoArray(CellStrings, TEXT("\t"), false);
 				
 					if ( bShowProgressDialog )
 					{
diff --git a/Engine/Source/Editor/PropertyEditor/Private/SDetailsViewBase.cpp b/Engine/Source/Editor/PropertyEditor/Private/SDetailsViewBase.cpp
--- a/Engine/Source/Editor/PropertyEditor/Private/SDetailsViewBase.cpp
+++ b/Engine/Source/Editor/PropertyEditor/Private/SDetailsViewBase.cpp
@@ -460,7 +460,7 @@ void SDetailsViewBase::FilterView(const FString& InFilterText)
 	// Remove whitespace from the front and back of the string
 	ParseString.Trim();
 	ParseString.TrimTrailing();
-	ParseString.ParseIntoArray(&CurrentFilterStrings, TEXT(" "), true);
+	ParseString.ParseIntoArray(CurrentFilterStrings, TEXT(" "), true);
 
 	bHasActiveFilter = CurrentFilterStrings.Num() > 0;
 
@@ -922,7 +922,7 @@ void SDetailsViewBase::RestoreExpandedItems(TSharedPtr<FPropertyNode> InitialSta
 	{
 		GConfig->GetString(TEXT("DetailCustomWidgetExpansion"), *BestBaseStruct->GetName(), ExpandedCustomItems, GEditorUserSettingsIni);
 		TArray<FString> ExpandedCustomItemsArray;
-		ExpandedCustomItems.ParseIntoArray(&ExpandedCustomItemsArray, TEXT(","), true);
+		ExpandedCustomItems.ParseIntoArray(ExpandedCustomItemsArray, TEXT(","), true);
 
 		ExpandedDetailNodes.Append(ExpandedCustomItemsArray);
 	}
diff --git a/Engine/Source/Editor/PropertyEditor/Private/SPropertyTreeViewImpl.cpp b/Engine/Source/Editor/PropertyEditor/Private/SPropertyTreeViewImpl.cpp
--- a/Engine/Source/Editor/PropertyEditor/Private/SPropertyTreeViewImpl.cpp
+++ b/Engine/Source/Editor/PropertyEditor/Private/SPropertyTreeViewImpl.cpp
@@ -966,7 +966,7 @@ void SPropertyTreeViewImpl::FilterView( const FString& InFilterText )
 	// Remove whitespace from the front and back of the string
 	ParseString.Trim();
 	ParseString.TrimTrailing();
-	ParseString.ParseIntoArray(&FilterStrings, TEXT(" "), true);
+	ParseString.ParseIntoArray(FilterStrings, TEXT(" "), true);
 
 	RootPropertyNode->FilterNodes( FilterStrings );
 	RootPropertyNode->ProcessSeenFlags(true);
diff --git a/Engine/Source/Editor/PropertyEditor/Private/UserInterface/PropertyTable/ObjectNameTableCellPresenter.h b/Engine/Source/Editor/PropertyEditor/Private/UserInterface/PropertyTable/ObjectNameTableCellPresenter.h
--- a/Engine/Source/Editor/PropertyEditor/Private/UserInterface/PropertyTable/ObjectNameTableCellPresenter.h
+++ b/Engine/Source/Editor/PropertyEditor/Private/UserInterface/PropertyTable/ObjectNameTableCellPresenter.h
@@ -72,7 +72,7 @@ private:
 		const FString& DisplayNameText = Cell->GetValueAsString();
 
 		TArray< FString > DisplayNamePieces;
-		DisplayNameText.ParseIntoArray( /*OUT*/ &DisplayNamePieces, TEXT("->"), true /*Cull Empty*/);
+		DisplayNameText.ParseIntoArray(DisplayNamePieces, TEXT("->"), true);
 
 		for (int Index = 0; Index < DisplayNamePieces.Num(); Index++)
 		{
diff --git a/Engine/Source/Editor/PropertyEditor/Private/UserInterface/PropertyTable/SColumnHeader.h b/Engine/Source/Editor/PropertyEditor/Private/UserInterface/PropertyTable/SColumnHeader.h
--- a/Engine/Source/Editor/PropertyEditor/Private/UserInterface/PropertyTable/SColumnHeader.h
+++ b/Engine/Source/Editor/PropertyEditor/Private/UserInterface/PropertyTable/SColumnHeader.h
@@ -90,7 +90,7 @@ protected:
 		const FText& DisplayNameText = Column->GetDisplayName();
 
 		TArray< FString > DisplayNamePieces;
-		DisplayNameText.ToString().ParseIntoArray( /*OUT*/ &DisplayNamePieces, TEXT("->"), true /*Cull Empty*/);
+		DisplayNameText.ToString().ParseIntoArray(DisplayNamePieces, TEXT("->"), true);
 
 		for (int Index = 0; Index < DisplayNamePieces.Num(); Index++)
 		{
diff --git a/Engine/Source/Editor/Sequencer/Private/SequencerNodeTree.cpp b/Engine/Source/Editor/Sequencer/Private/SequencerNodeTree.cpp
--- a/Engine/Source/Editor/Sequencer/Private/SequencerNodeTree.cpp
+++ b/Engine/Source/Editor/Sequencer/Private/SequencerNodeTree.cpp
@@ -279,7 +279,7 @@ void FSequencerNodeTree::FilterNodes( const FString& InFilter )
 		FilterString.Trim();
 		FilterString.TrimTrailing();
 		const bool bCullEmpty = true;
-		FilterString.ParseIntoArray( &FilterStrings, TEXT(" "), bCullEmpty );
+		FilterString.ParseIntoArray( FilterStrings, TEXT(" "), bCullEmpty );
 
 		for( auto It = ObjectBindingMap.CreateIterator(); It; ++It )
 		{
diff --git a/Engine/Source/Editor/UnrealEd/Private/AssetSelection.cpp b/Engine/Source/Editor/UnrealEd/Private/AssetSelection.cpp
--- a/Engine/Source/Editor/UnrealEd/Private/AssetSelection.cpp
+++ b/Engine/Source/Editor/UnrealEd/Private/AssetSelection.cpp
@@ -462,7 +462,7 @@ namespace AssetUtil
 			{
 				TArray<FString> DroppedAssetStrings;
 				const TCHAR AssetDelimiter[] = { AssetMarshalDefs::AssetDelimiter, TEXT('\0') };
-				DragDropOp->GetText().ParseIntoArray(&DroppedAssetStrings, AssetDelimiter, true);
+				DragDropOp->GetText().ParseIntoArray(DroppedAssetStrings, AssetDelimiter, true);
 
 				FAssetRegistryModule& AssetRegistryModule = FModuleManager::LoadModuleChecked<FAssetRegistryModule>(TEXT("AssetRegistry"));
 				IAssetRegistry& AssetRegistry = AssetRegistryModule.Get();
diff --git a/Engine/Source/Editor/UnrealEd/Private/Commandlets/ContentCommandlets.cpp b/Engine/Source/Editor/UnrealEd/Private/Commandlets/ContentCommandlets.cpp
--- a/Engine/Source/Editor/UnrealEd/Private/Commandlets/ContentCommandlets.cpp
+++ b/Engine/Source/Editor/UnrealEd/Private/Commandlets/ContentCommandlets.cpp
@@ -175,7 +175,7 @@ int32 UResavePackagesCommandlet::InitializeResaveParameters( const TArray<FStrin
 		if ( FParse::Value(*CurrentSwitch, TEXT("RESAVECLASS="), ClassList, false) )
 		{
 			TArray<FString> ClassNames;
-			ClassList.ParseIntoArray(&ClassNames, TEXT(","), true);
+			ClassList.ParseIntoArray(ClassNames, TEXT(","), true);
 			for ( int32 Idx = 0; Idx < ClassNames.Num(); Idx++ )
 			{
 				ResaveClasses.AddUnique(*ClassNames[Idx]);
diff --git a/Engine/Source/Editor/UnrealEd/Private/Commandlets/CookCommandlet.cpp b/Engine/Source/Editor/UnrealEd/Private/Commandlets/CookCommandlet.cpp
--- a/Engine/Source/Editor/UnrealEd/Private/Commandlets/CookCommandlet.cpp
+++ b/Engine/Source/Editor/UnrealEd/Private/Commandlets/CookCommandlet.cpp
@@ -1092,7 +1092,7 @@ bool UCookCommandlet::NewCook( const TArray<ITargetPlatform*>& Platforms, TArray
 	{
 		if (SectionStr.Contains(TEXT("+")))
 		{
-			SectionStr.ParseIntoArray(&CmdLineIniSections,TEXT("+"),true);
+			SectionStr.ParseIntoArray(CmdLineIniSections,TEXT("+"),true);
 		}
 		else
 		{
diff --git a/Engine/Source/Editor/UnrealEd/Private/Commandlets/GatherTextCommandletBase.cpp b/Engine/Source/Editor/UnrealEd/Private/Commandlets/GatherTextCommandletBase.cpp
--- a/Engine/Source/Editor/UnrealEd/Private/Commandlets/GatherTextCommandletBase.cpp
+++ b/Engine/Source/Editor/UnrealEd/Private/Commandlets/GatherTextCommandletBase.cpp
@@ -497,12 +497,12 @@ bool FWordCountReportData::FromCSV( const FString& InString )
 	// Note, we do a trivial parse here and assume every delimiter we see separates entries/rows and is not contained within an entry
 	TArray<FString> CSVRows;
 
-	int32 NumRows = CSVString.ParseIntoArray( &CSVRows, NewLineDelimiter, true );
+	int32 NumRows = CSVString.ParseIntoArray( CSVRows, NewLineDelimiter, true );
 
 	for(int32 RowIdx = 0; RowIdx < NumRows; RowIdx++ )
 	{
 		TArray< FString > RowEntriesArray;
-		CSVRows[RowIdx].ParseIntoArray( &RowEntriesArray, EntryDelimiter, false);
+		CSVRows[RowIdx].ParseIntoArray( RowEntriesArray, EntryDelimiter, false);
 		Data.Add(RowEntriesArray);
 	}
 
diff --git a/Engine/Source/Editor/UnrealEd/Private/Commandlets/GenerateBlueprintAPICommandlet.cpp b/Engine/Source/Editor/UnrealEd/Private/Commandlets/GenerateBlueprintAPICommandlet.cpp
--- a/Engine/Source/Editor/UnrealEd/Private/Commandlets/GenerateBlueprintAPICommandlet.cpp
+++ b/Engine/Source/Editor/UnrealEd/Private/Commandlets/GenerateBlueprintAPICommandlet.cpp
@@ -791,7 +791,7 @@ static void GenerateBlueprintAPIUtils::DumpActionMenuItem(uint32 Indent, FGraphA
 	}
 
 	TArray<FString> Categories;
-	ActionCategory.ParseIntoArray(&Categories, TEXT("|"), true);
+	ActionCategory.ParseIntoArray(Categories, TEXT("|"), true);
 
 	for (const FString& Category : Categories)
 	{
diff --git a/Engine/Source/Editor/UnrealEd/Private/Commandlets/InternationalizationExportCommandlet.cpp b/Engine/Source/Editor/UnrealEd/Private/Commandlets/InternationalizationExportCommandlet.cpp
--- a/Engine/Source/Editor/UnrealEd/Private/Commandlets/InternationalizationExportCommandlet.cpp
+++ b/Engine/Source/Editor/UnrealEd/Private/Commandlets/InternationalizationExportCommandlet.cpp
@@ -470,7 +470,7 @@ bool FPortableObjectHeader::FromLocPOEntry( const TSharedRef<const FPortableObje
 
 	// The POEntry would store header info inside the MsgStr[0]
 	TArray<FString> HeaderLinesToProcess;
-	LocEntry->MsgStr[0].ReplaceEscapedCharWithChar().ParseIntoArray( &HeaderLinesToProcess, NewLineDelimiter, true );
+	LocEntry->MsgStr[0].ReplaceEscapedCharWithChar().ParseIntoArray( HeaderLinesToProcess, NewLineDelimiter, true );
 
 	for( const FString& PotentialHeaderEntry : HeaderLinesToProcess )
 	{
@@ -569,7 +569,7 @@ bool FPortableObjectFormatDOM::FromString( const FString& InStr )
 	FString ParseString = InStr.Replace(TEXT("\r\n"), NewLineDelimiter);
 
 	TArray<FString> LinesToProcess;
-	ParseString.ParseIntoArray( &LinesToProcess, NewLineDelimiter, false );
+	ParseString.ParseIntoArray( LinesToProcess, NewLineDelimiter, false );
 
 	TSharedRef<FPortableObjectEntry> ProcessedEntry = MakeShareable( new FPortableObjectEntry );
 	bool bHasMsgId = false;
@@ -932,7 +932,7 @@ void FPortableObjectEntry::AddExtractedComment( const FString& InComment )
 
 	//// Extracted comments can contain multiple references in a single line so we parse those out.
 	//TArray<FString> CommentsToProcess;
-	//InComment.ParseIntoArray( &CommentsToProcess, TEXT(" "), true );
+	//InComment.ParseIntoArray( CommentsToProcess, TEXT(" "), true );
 	//for( const FString& ExtractedComment : CommentsToProcess )
 	//{
 	//	ExtractedComments.AddUnique( ExtractedComment );
@@ -948,7 +948,7 @@ void FPortableObjectEntry::AddReference( const FString& InReference )
 	
 	//// Reference comments can contain multiple references in a single line so we parse those out.
 	//TArray<FString> ReferencesToProcess;
-	//InReference.ParseIntoArray( &ReferencesToProcess, TEXT(" "), true );
+	//InReference.ParseIntoArray( ReferencesToProcess, TEXT(" "), true );
 	//for( const FString& Reference : ReferencesToProcess )
 	//{
 	//	ReferenceComments.AddUnique( Reference );
diff --git a/Engine/Source/Editor/UnrealEd/Private/Commandlets/PackageUtilities.cpp b/Engine/Source/Editor/UnrealEd/Private/Commandlets/PackageUtilities.cpp
--- a/Engine/Source/Editor/UnrealEd/Private/Commandlets/PackageUtilities.cpp
+++ b/Engine/Source/Editor/UnrealEd/Private/Commandlets/PackageUtilities.cpp
@@ -1009,7 +1009,7 @@ void FPkgInfoReporter_Log::GeneratePackageReport( ULinkerLoad* InLinker/*=NULL*/
 		if ( FParse::Value(FCommandLine::Get(), TEXT("SORT="), SortingParms) )
 		{
 			TArray<FString> SortValues;
-			SortingParms.ParseIntoArray(&SortValues, TEXT(","), true);
+			SortingParms.ParseIntoArray(SortValues, TEXT(","), true);
 
 			for ( int32 i = 0; i < EXPORTSORT_MAX; i++ )
 			{
diff --git a/Engine/Source/Editor/UnrealEd/Private/CookOnTheFlyServer.cpp b/Engine/Source/Editor/UnrealEd/Private/CookOnTheFlyServer.cpp
--- a/Engine/Source/Editor/UnrealEd/Private/CookOnTheFlyServer.cpp
+++ b/Engine/Source/Editor/UnrealEd/Private/CookOnTheFlyServer.cpp
@@ -2014,7 +2014,7 @@ bool UCookOnTheFlyServer::GetCurrentIniVersionStrings( const ITargetPlatform* Ta
 	{
 
 		TArray<FString> IniVersionedArray;
-		IniVersioned.ParseIntoArray(&IniVersionedArray, TEXT(":"), false);
+		IniVersioned.ParseIntoArray(IniVersionedArray, TEXT(":"), false);
 
 		if ( IniVersionedArray.Num() != 3 )
 		{
diff --git a/Engine/Source/Editor/UnrealEd/Private/Editor.cpp b/Engine/Source/Editor/UnrealEd/Private/Editor.cpp
--- a/Engine/Source/Editor/UnrealEd/Private/Editor.cpp
+++ b/Engine/Source/Editor/UnrealEd/Private/Editor.cpp
@@ -3052,7 +3052,7 @@ void UEditorEngine::ParseMapSectionIni(const TCHAR* InCmdParams, TArray<FString>
 		if (SectionStr.Contains(TEXT("+")))
 		{
 			TArray<FString> Sections;
-			SectionStr.ParseIntoArray(&Sections,TEXT("+"),true);
+			SectionStr.ParseIntoArray(Sections,TEXT("+"),true);
 			for (int32 Index = 0; Index < Sections.Num(); Index++)
 			{
 				LoadMapListFromIni(Sections[Index], OutMapList);
@@ -6588,7 +6588,7 @@ void UEditorEngine::UpdateAutoLoadProject()
 		FPlatformMisc::GetOSVersions(OSVersion, OSSubVersion);
 		
 		TArray<FString> Components;
-		OSVersion.ParseIntoArray(&Components, TEXT("."), true);
+		OSVersion.ParseIntoArray(Components, TEXT("."), true);
 		uint8 ComponentValues[3] = {0};
 		
 		for(uint32 i = 0; i < Components.Num() && i < 3; i++)
diff --git a/Engine/Source/Editor/UnrealEd/Private/EditorCategoryUtils.cpp b/Engine/Source/Editor/UnrealEd/Private/EditorCategoryUtils.cpp
--- a/Engine/Source/Editor/UnrealEd/Private/EditorCategoryUtils.cpp
+++ b/Engine/Source/Editor/UnrealEd/Private/EditorCategoryUtils.cpp
@@ -294,7 +294,7 @@ void FEditorCategoryUtils::GetClassHideCategories(UClass const* Class, TArray<FS
 	{
 		FString const& HideCategories = Class->GetMetaData(ClassHideCategoriesMetaKey);
 
-		HideCategories.ParseIntoArray(&CategoriesOut, TEXT(" "), /*InCullEmpty =*/true);
+		HideCategories.ParseIntoArray(CategoriesOut, TEXT(" "), /*InCullEmpty =*/true);
 		
 		for (FString& Category : CategoriesOut)
 		{
@@ -312,7 +312,7 @@ void  FEditorCategoryUtils::GetClassShowCategories(UClass const* Class, TArray<F
 	if (Class->HasMetaData(ClassShowCategoriesMetaKey))
 	{
 		FString const& ShowCategories = Class->GetMetaData(ClassShowCategoriesMetaKey);
-		ShowCategories.ParseIntoArray(&CategoriesOut, TEXT(" "), /*InCullEmpty =*/true);
+		ShowCategories.ParseIntoArray(CategoriesOut, TEXT(" "), /*InCullEmpty =*/true);
 
 		for (FString& Category : CategoriesOut)
 		{
@@ -357,7 +357,7 @@ bool FEditorCategoryUtils::IsCategoryHiddenFromClass(UClass const* Class, FStrin
 		else // see if the category's root is hidden
 		{
 			TArray<FString> SubCategoryList;
-			DisplayCategory.ParseIntoArray(&SubCategoryList, TEXT("|"), /*InCullEmpty =*/true);
+			DisplayCategory.ParseIntoArray(SubCategoryList, TEXT("|"), /*InCullEmpty =*/true);
 
 			FString FullSubCategoryPath;
 			for (FString const& SubCategory : SubCategoryList)
diff --git a/Engine/Source/Editor/UnrealEd/Private/FeaturePack/FeaturePack.cpp b/Engine/Source/Editor/UnrealEd/Private/FeaturePack/FeaturePack.cpp
--- a/Engine/Source/Editor/UnrealEd/Private/FeaturePack/FeaturePack.cpp
+++ b/Engine/Source/Editor/UnrealEd/Private/FeaturePack/FeaturePack.cpp
@@ -40,7 +40,7 @@ void FFeaturePack::ParsePacks()
 	{
 		PackData EachPackData;
 		TArray<FString> PackEntries;
-		PacksToAdd[iPackEntry].ParseIntoArray(&PackEntries, TEXT(","), true);
+		PacksToAdd[iPackEntry].ParseIntoArray(PackEntries, TEXT(","), true);
 		FString PackSource;
 		FString PackName;
 		// Parse the pack name and source
diff --git a/Engine/Source/Editor/UnrealEd/Private/Kismet2/BlueprintEditorUtils.cpp b/Engine/Source/Editor/UnrealEd/Private/Kismet2/BlueprintEditorUtils.cpp
--- a/Engine/Source/Editor/UnrealEd/Private/Kismet2/BlueprintEditorUtils.cpp
+++ b/Engine/Source/Editor/UnrealEd/Private/Kismet2/BlueprintEditorUtils.cpp
@@ -6309,7 +6309,7 @@ TSharedRef<SWidget> FBlueprintEditorUtils::ConstructBlueprintInterfaceClassPicke
 			FString const& ProhibitedList = Blueprint->ParentClass->GetMetaData(FBlueprintMetadata::MD_ProhibitedInterfaces);
 
 			TArray<FString> ProhibitedInterfaceNames;
-			ProhibitedList.ParseIntoArray(&ProhibitedInterfaceNames, TEXT(","), true);
+			ProhibitedList.ParseIntoArray(ProhibitedInterfaceNames, TEXT(","), true);
 
 			// loop over all the prohibited interfaces
 			for (int32 ExclusionIndex = 0; ExclusionIndex < ProhibitedInterfaceNames.Num(); ++ExclusionIndex)
diff --git a/Engine/Source/Editor/UnrealEd/Private/Kismet2/CompilerResultsLog.cpp b/Engine/Source/Editor/UnrealEd/Private/Kismet2/CompilerResultsLog.cpp
--- a/Engine/Source/Editor/UnrealEd/Private/Kismet2/CompilerResultsLog.cpp
+++ b/Engine/Source/Editor/UnrealEd/Private/Kismet2/CompilerResultsLog.cpp
@@ -344,7 +344,7 @@ TArray< TSharedRef<FTokenizedMessage> > FCompilerResultsLog::ParseCompilerLogDum
 	TArray< TSharedRef<FTokenizedMessage> > Messages;
 
 	TArray< FString > MessageLines;
-	LogDump.ParseIntoArray(&MessageLines, TEXT("\n"), false);
+	LogDump.ParseIntoArray(MessageLines, TEXT("\n"), false);
 
 	// delete any trailing empty lines
 	for (int32 i = MessageLines.Num()-1; i >= 0; --i)
diff --git a/Engine/Source/Editor/UnrealEd/Private/Kismet2/Kismet2.cpp b/Engine/Source/Editor/UnrealEd/Private/Kismet2/Kismet2.cpp
--- a/Engine/Source/Editor/UnrealEd/Private/Kismet2/Kismet2.cpp
+++ b/Engine/Source/Editor/UnrealEd/Private/Kismet2/Kismet2.cpp
@@ -1769,7 +1769,7 @@ bool FKismetEditorUtilities::CanBlueprintImplementInterface(UBlueprint const* Bl
 			FString const& ProhibitedList = Blueprint->ParentClass->GetMetaData(FBlueprintMetadata::MD_ProhibitedInterfaces);
 			
 			TArray<FString> ProhibitedInterfaceNames;
-			ProhibitedList.ParseIntoArray(&ProhibitedInterfaceNames, TEXT(","), true);
+			ProhibitedList.ParseIntoArray(ProhibitedInterfaceNames, TEXT(","), true);
 
 			FString const& InterfaceName = Class->GetName();
 			// loop over all the prohibited interfaces
diff --git a/Engine/Source/Editor/UnrealEd/Private/Niagara/EdGraphSchema_Niagara.cpp b/Engine/Source/Editor/UnrealEd/Private/Niagara/EdGraphSchema_Niagara.cpp
--- a/Engine/Source/Editor/UnrealEd/Private/Niagara/EdGraphSchema_Niagara.cpp
+++ b/Engine/Source/Editor/UnrealEd/Private/Niagara/EdGraphSchema_Niagara.cpp
@@ -302,7 +302,7 @@ void UEdGraphSchema_Niagara::GetPinDefaultValue(UEdGraphPin* Pin, FVector4& OutD
 	FString ValueString = Pin->DefaultValue.IsEmpty() ? Pin->AutogeneratedDefaultValue : Pin->DefaultValue;
 	ValueString.Trim();
 	ValueString.TrimTrailing();
-	ValueString.ParseIntoArray(&ResultString, TEXT(","), true);
+	ValueString.ParseIntoArray(ResultString, TEXT(","), true);
 	if (ResultString.Num() == 4)
 	{
 		OutDefault = FVector4(FCString::Atof(*ResultString[0]), FCString::Atof(*ResultString[1]), FCString::Atof(*ResultString[2]), FCString::Atof(*ResultString[3]));
@@ -315,7 +315,7 @@ void UEdGraphSchema_Niagara::GetPinDefaultValue(UEdGraphPin* Pin, FMatrix& OutDe
 	FString ValueString = Pin->DefaultValue.IsEmpty() ? Pin->AutogeneratedDefaultValue : Pin->DefaultValue;
 	ValueString.Trim();
 	ValueString.TrimTrailing();
-	ValueString.ParseIntoArray(&ResultString, TEXT(","), true);
+	ValueString.ParseIntoArray(ResultString, TEXT(","), true);
 	check(ResultString.Num() == 16);
 	if (ResultString.Num() == 16)
 	{
diff --git a/Engine/Source/Editor/UnrealEd/Private/ObjectTools.cpp b/Engine/Source/Editor/UnrealEd/Private/ObjectTools.cpp
--- a/Engine/Source/Editor/UnrealEd/Private/ObjectTools.cpp
+++ b/Engine/Source/Editor/UnrealEd/Private/ObjectTools.cpp
@@ -2723,7 +2723,7 @@ namespace ObjectTools
 
 			// Parse the format into its extension and description parts
 			TArray<FString> FormatComponents;
-			CurFormat.ParseIntoArray( &FormatComponents, TEXT(";"), false );
+			CurFormat.ParseIntoArray( FormatComponents, TEXT(";"), false );
 
 			for ( int32 ComponentIndex = 0; ComponentIndex < FormatComponents.Num(); ComponentIndex += 2 )
 			{
diff --git a/Engine/Source/Editor/UnrealEd/Private/SKeySelector.cpp b/Engine/Source/Editor/UnrealEd/Private/SKeySelector.cpp
--- a/Engine/Source/Editor/UnrealEd/Private/SKeySelector.cpp
+++ b/Engine/Source/Editor/UnrealEd/Private/SKeySelector.cpp
@@ -344,7 +344,7 @@ void SKeySelector::OnFilterTextCommitted(const FText& NewText, ETextCommit::Type
 
 void SKeySelector::GetSearchTokens(const FString& SearchString, TArray<FString>& OutTokens) const
 {
-	if (SearchString.Contains("\"") && SearchString.ParseIntoArray(&OutTokens, TEXT("\""), true) > 0)
+	if (SearchString.Contains("\"") && SearchString.ParseIntoArray(OutTokens, TEXT("\""), true) > 0)
 	{
 		for (auto &TokenIt : OutTokens)
 		{
@@ -368,7 +368,7 @@ void SKeySelector::GetSearchTokens(const FString& SearchString, TArray<FString>&
 	else
 	{
 		// unquoted search equivalent to a match-any-of search
-		SearchString.ParseIntoArray(&OutTokens, TEXT(" "), true);
+		SearchString.ParseIntoArray(OutTokens, TEXT(" "), true);
 	}
 }
 
diff --git a/Engine/Source/Editor/UnrealEd/Private/Tests/AutomationEditorCommon.cpp b/Engine/Source/Editor/UnrealEd/Private/Tests/AutomationEditorCommon.cpp
--- a/Engine/Source/Editor/UnrealEd/Private/Tests/AutomationEditorCommon.cpp
+++ b/Engine/Source/Editor/UnrealEd/Private/Tests/AutomationEditorCommon.cpp
@@ -320,7 +320,7 @@ namespace AutomationEditorCommonUtils
 				}
 
 				TArray<FString> PropertyChain;
-				FactorySettings[i].SettingName.ParseIntoArray(&PropertyChain, TEXT("."), false);
+				FactorySettings[i].SettingName.ParseIntoArray(PropertyChain, TEXT("."), false);
 				ApplyCustomFactorySetting(InFactory, PropertyChain, FactorySettings[i].Value);
 			}
 		}
@@ -420,7 +420,7 @@ namespace AutomationEditorCommonUtils
 		{
 			UE_LOG(LogEditorAutomationTests, VeryVerbose, TEXT("Loading and parsing the data from '%s' into an array."), *InFileLocation);
 			FFileHelper::LoadFileToString(RawData, *InFileLocation);
-			RawData.ParseIntoArray(&OutArray, TEXT(","), false);
+			RawData.ParseIntoArray(OutArray, TEXT(","), false);
 		}
 
 		UE_LOG(LogEditorAutomationTests, Warning, TEXT("Unable to create an array.  '%s' does not exist."), *InFileLocation);
diff --git a/Engine/Source/Editor/UnrealEd/Private/Tests/EditorOpenAssetAutomationTests.cpp b/Engine/Source/Editor/UnrealEd/Private/Tests/EditorOpenAssetAutomationTests.cpp
--- a/Engine/Source/Editor/UnrealEd/Private/Tests/EditorOpenAssetAutomationTests.cpp
+++ b/Engine/Source/Editor/UnrealEd/Private/Tests/EditorOpenAssetAutomationTests.cpp
@@ -88,7 +88,7 @@ void FOpenAssetEditors::GetTests(TArray<FString>& OutBeautifiedNames, TArray <FS
 	FString EngineContentFolderLocation = FPaths::ConvertRelativePathToFull(*FPaths::EngineContentDir());
 	//Put the Engine Content Folder Location into an array.  
 	TArray<FString> EngineContentFolderLocationArray;
-	EngineContentFolderLocation.ParseIntoArray(&EngineContentFolderLocationArray, TEXT("/"), true);
+	EngineContentFolderLocation.ParseIntoArray(EngineContentFolderLocationArray, TEXT("/"), true);
 
 	for (int32 i = 0; i < AssetNames.Num(); ++i)
 	{
@@ -100,7 +100,7 @@ void FOpenAssetEditors::GetTests(TArray<FString>& OutBeautifiedNames, TArray <FS
 
 		//Put File location into an array.  
 		TArray<FString> FilenameArray;
-		Filename.ParseIntoArray(&FilenameArray, TEXT("/"), true);
+		Filename.ParseIntoArray(FilenameArray, TEXT("/"), true);
 
 		//Loop through the location array's and compare each element.  
 		//The loop runs while the index is less than the number of elements in the EngineContentFolderLocation array.
diff --git a/Engine/Source/Editor/UnrealEd/Private/UnrealEdSrv.cpp b/Engine/Source/Editor/UnrealEd/Private/UnrealEdSrv.cpp
--- a/Engine/Source/Editor/UnrealEd/Private/UnrealEdSrv.cpp
+++ b/Engine/Source/Editor/UnrealEd/Private/UnrealEdSrv.cpp
@@ -1147,7 +1147,7 @@ bool UUnrealEdEngine::Exec( UWorld* InWorld, const TCHAR* Stream, FOutputDevice&
 		if (FParse::Value(Str, TEXT("Find="), SearchTermStr, false))
 		{
 			TArray<FString> SearchTerms;
-			SearchTermStr.ParseIntoArray( &SearchTerms, TEXT(","), true );
+			SearchTermStr.ParseIntoArray( SearchTerms, TEXT(","), true );
 
 			TArray<UObject*> ModifiedObjects;
 			if( SearchTerms.Num() )
diff --git a/Engine/Source/Programs/BuildPatchTool/Private/BuildPatchToolMain.cpp b/Engine/Source/Programs/BuildPatchTool/Private/BuildPatchToolMain.cpp
--- a/Engine/Source/Programs/BuildPatchTool/Private/BuildPatchToolMain.cpp
+++ b/Engine/Source/Programs/BuildPatchTool/Private/BuildPatchToolMain.cpp
@@ -432,7 +432,7 @@ int32 BuildPatchToolMain( const TCHAR* CommandLine )
 		TArray<FString> ManifestsArr;
 		if (ManifestsList.Len() > 0)
 		{
-			ManifestsList.ParseIntoArray(&ManifestsArr, TEXT(","), true);
+			ManifestsList.ParseIntoArray(ManifestsArr, TEXT(","), true);
 		}
 		else if (ManifestsFile.Len() > 0)
 		{
@@ -441,7 +441,7 @@ int32 BuildPatchToolMain( const TCHAR* CommandLine )
 			if (FFileHelper::LoadFileToString(Temp, *ManifestsFilePath))
 			{
 				Temp.ReplaceInline(TEXT("\r"), TEXT("\n"));
-				Temp.ParseIntoArray(&ManifestsArr, TEXT("\n"), true);
+				Temp.ParseIntoArray(ManifestsArr, TEXT("\n"), true);
 			}
 			else
 			{
diff --git a/Engine/Source/Programs/CrashReporter/CrashReportClient/Private/CrashDescription.cpp b/Engine/Source/Programs/CrashReporter/CrashReportClient/Private/CrashDescription.cpp
--- a/Engine/Source/Programs/CrashReporter/CrashReportClient/Private/CrashDescription.cpp
+++ b/Engine/Source/Programs/CrashReporter/CrashReportClient/Private/CrashDescription.cpp
@@ -84,7 +84,7 @@ FCrashDescription::FCrashDescription( FString WERXMLFilepath ) :
 				const FString Parameter8Value = Parameter8Node->GetContent();
 
 				TArray<FString> ParsedParameters8;
-				Parameter8Value.ParseIntoArray( &ParsedParameters8, TEXT( "!" ), false );
+				Parameter8Value.ParseIntoArray( ParsedParameters8, TEXT( "!" ), false );
 
 				if( ParsedParameters8.Num() > 1 )
 				{
@@ -103,7 +103,7 @@ FCrashDescription::FCrashDescription( FString WERXMLFilepath ) :
 				const FString Parameter9Value = Parameter9Node->GetContent();
 
 				TArray<FString> ParsedParameters9;
-				Parameter9Value.ParseIntoArray( &ParsedParameters9, TEXT( "!" ), false );
+				Parameter9Value.ParseIntoArray( ParsedParameters9, TEXT( "!" ), false );
 
 				if( ParsedParameters9.Num() > 0 )
 				{
@@ -122,7 +122,7 @@ FCrashDescription::FCrashDescription( FString WERXMLFilepath ) :
 					const FString BaseDir = ParsedParameters9[1];
 
 					TArray<FString> SubDirs;
-					BaseDir.ParseIntoArray( &SubDirs, TEXT( "/" ), true );
+					BaseDir.ParseIntoArray( SubDirs, TEXT( "/" ), true );
 					const int SubDirsNum = SubDirs.Num();
 					const FString PlatformName = SubDirsNum > 0 ? SubDirs[SubDirsNum - 1] : TEXT("");
 					if( Product.Len() > 0 )
@@ -178,7 +178,7 @@ void FCrashDescription::InitializeEngineVersion()
 	uint16 Patch = 0;
 
 	TArray<FString> ParsedBuildVersion;
-	BuildVersion.ParseIntoArray( &ParsedBuildVersion, TEXT( "." ), false );
+	BuildVersion.ParseIntoArray( ParsedBuildVersion, TEXT( "." ), false );
 
 	if( ParsedBuildVersion.Num() >= 3 )
 	{
