Added implementation of AnnotateFile that takes a changelist

UE-8036 - Add SourceControlHelpers::AnnotateFile that will work with CL number

diff --git a/Engine/Source/Developer/SourceControl/Private/SourceControlHelpers.cpp b/Engine/Source/Developer/SourceControl/Private/SourceControlHelpers.cpp
--- a/Engine/Source/Developer/SourceControl/Private/SourceControlHelpers.cpp
+++ b/Engine/Source/Developer/SourceControl/Private/SourceControlHelpers.cpp
@@ -171,6 +171,33 @@ bool AnnotateFile( ISourceControlProvider& InProvider, const FString& InLabel, c
 	return false;
 }
 
+bool AnnotateFile( ISourceControlProvider& InProvider, int32 InCheckInIdentifier, const FString& InFile, TArray<FAnnotationLine>& OutLines )
+{
+	TSharedRef<FUpdateStatus, ESPMode::ThreadSafe> UpdateStatusOperation = ISourceControlOperation::Create<FUpdateStatus>();
+	UpdateStatusOperation->SetUpdateHistory(true);
+	if(InProvider.Execute(UpdateStatusOperation, InFile) == ECommandResult::Succeeded)
+	{
+		FSourceControlStatePtr State = InProvider.GetState(InFile, EStateCacheUsage::Use);
+		if(State.IsValid())
+		{
+			for(int32 HistoryIndex = State->GetHistorySize() - 1; HistoryIndex >= 0; HistoryIndex--)
+			{
+				// check that the changelist corresponds to this revision - we assume history is in latest-first order
+				TSharedPtr<ISourceControlRevision, ESPMode::ThreadSafe> Revision = State->GetHistoryItem(HistoryIndex);
+				if(Revision.IsValid() && Revision->GetCheckInIdentifier() >= InCheckInIdentifier)
+				{
+					if(Revision->GetAnnotated(OutLines))
+					{
+						return true;
+					}
+				}
+			}
+		}
+	}
+
+	return false;
+}
+
 bool CheckOutFile( const FString& InFilePath )
 {
 	if ( InFilePath.IsEmpty() )
@@ -274,14 +301,14 @@ bool CheckoutOrMarkForAdd( const FString& InDestFile, const FText& InFileDescrip
 		if (SourceControlState.IsValid())
 		{
 			if (SourceControlState->IsSourceControlled() && SourceControlState->CanCheckout())
-			{
-				ECommandResult::Type Result = Provider.Execute(ISourceControlOperation::Create<FCheckOut>(), InDestFile);
-				bSucceeded = (Result == ECommandResult::Succeeded);
-				if (!bSucceeded)
 				{
-					OutFailReason = FText::Format(LOCTEXT("SourceControlCheckoutError", "Could not check out {0} file."), InFileDescription);
+					ECommandResult::Type Result = Provider.Execute(ISourceControlOperation::Create<FCheckOut>(), InDestFile);
+					bSucceeded = (Result == ECommandResult::Succeeded);
+					if (!bSucceeded)
+					{
+						OutFailReason = FText::Format(LOCTEXT("SourceControlCheckoutError", "Could not check out {0} file."), InFileDescription);
+					}
 				}
-			}
 		}
 	}
 
diff --git a/Engine/Source/Developer/SourceControl/Public/SourceControlHelpers.h b/Engine/Source/Developer/SourceControl/Public/SourceControlHelpers.h
--- a/Engine/Source/Developer/SourceControl/Public/SourceControlHelpers.h
+++ b/Engine/Source/Developer/SourceControl/Public/SourceControlHelpers.h
@@ -73,7 +73,7 @@ namespace SourceControlHelpers
 	SOURCECONTROL_API extern void RevertUnchangedFiles( ISourceControlProvider& InProvider, const TArray<FString>& InFiles );
 
 	/**
-	 * Helper function to annotate a file
+	 * Helper function to annotate a file using a label
 	 * @param	InProvider	The provider to use
 	 * @param	InLabel		The label to use to retrieve the file
 	 * @param	InFile		The file to annotate
@@ -82,6 +82,16 @@ namespace SourceControlHelpers
 	 */
 	SOURCECONTROL_API extern bool AnnotateFile( ISourceControlProvider& InProvider, const FString& InLabel, const FString& InFile, TArray<FAnnotationLine>& OutLines );
 
+	/**
+	 * Helper function to annotate a file using a changelist/checkin identifier
+	 * @param	InProvider				The provider to use
+	 * @param	InCheckInIdentifier		The changelist/check identifier to use to retrieve the file
+	 * @param	InFile					The file to annotate
+	 * @param	OutLines				Output array of annotated lines
+	 * @returns true if successful
+	 */
+	SOURCECONTROL_API extern bool AnnotateFile( ISourceControlProvider& InProvider, int32 InCheckInIdentifier, const FString& InFile, TArray<FAnnotationLine>& OutLines );
+
 	/**
 	 * Helper function to check out a file
 	 * @param	InFile		The file path to check in
