Removed call to FindModules from FSourceControlModule

This was older code from before we had modular features, and is no longer required as modular features already load and register themselves.

diff --git a/Engine/Source/Developer/SourceControl/Private/SSourceControlPicker.cpp b/Engine/Source/Developer/SourceControl/Private/SSourceControlPicker.cpp
--- a/Engine/Source/Developer/SourceControl/Private/SSourceControlPicker.cpp
+++ b/Engine/Source/Developer/SourceControl/Private/SSourceControlPicker.cpp
@@ -60,8 +60,6 @@ TSharedRef<SWidget> SSourceControlPicker::OnGetMenuContent() const
 {
 	FSourceControlModule& SourceControlModule = FSourceControlModule::Get();
 
-	SourceControlModule.RefreshSourceControlProviders();
-
 	FMenuBuilder MenuBuilder(true, NULL);
 
 	// Get the provider names first so that we can sort them for the UI
diff --git a/Engine/Source/Developer/SourceControl/Private/SourceControlModule.cpp b/Engine/Source/Developer/SourceControl/Private/SourceControlModule.cpp
--- a/Engine/Source/Developer/SourceControl/Private/SourceControlModule.cpp
+++ b/Engine/Source/Developer/SourceControl/Private/SourceControlModule.cpp
@@ -179,23 +179,8 @@ void FSourceControlModule::OnSourceControlDialogClosed(const TSharedRef<SWindow>
 #endif
 }
 
-void FSourceControlModule::RefreshSourceControlProviders() const
-{
-	TArray<FName> ModuleNames;
-	FModuleManager::Get().FindModules(TEXT("*SourceControl"), ModuleNames);
-	ModuleNames.Remove(FName(TEXT("SourceControl")));
-
-	for(int32 ModuleIndex = 0; ModuleIndex < ModuleNames.Num(); ModuleIndex++)
-	{
-		FModuleManager::Get().LoadModule(ModuleNames[ModuleIndex]);
-	}
-}
-
 void FSourceControlModule::InitializeSourceControlProviders()
 {
-	// Look for valid SourceControl modules - they will register themselves as editor features
-	RefreshSourceControlProviders();
-
 	int32 SourceControlCount = IModularFeatures::Get().GetModularFeatureImplementationCount(SourceControlFeatureName);
 	if( SourceControlCount > 0 )
 	{
diff --git a/Engine/Source/Developer/SourceControl/Private/SourceControlModule.h b/Engine/Source/Developer/SourceControl/Private/SourceControlModule.h
--- a/Engine/Source/Developer/SourceControl/Private/SourceControlModule.h
+++ b/Engine/Source/Developer/SourceControl/Private/SourceControlModule.h
@@ -37,11 +37,6 @@ public:
 	/** Save the settings to the ini file */
 	void SaveSettings();
 
-	/**
-	 * Refresh the available source control providers.
-	 */
-	void RefreshSourceControlProviders() const;
-
 	/**
 	 * Get the number of currently registered source control providers.
 	 */
