Guard against source control crashes if a provider becomes null

diff --git a/Engine/Source/Developer/SourceControl/Private/SourceControlModule.cpp b/Engine/Source/Developer/SourceControl/Private/SourceControlModule.cpp
--- a/Engine/Source/Developer/SourceControl/Private/SourceControlModule.cpp
+++ b/Engine/Source/Developer/SourceControl/Private/SourceControlModule.cpp
@@ -214,35 +214,38 @@ void FSourceControlModule::InitializeSourceControlProviders()
 
 void FSourceControlModule::Tick()
 {	
-	ISourceControlProvider& Provider = GetProvider();
-
-	// tick the provider, so any operation results can be read back
-	Provider.Tick();
-	
-	// don't allow background status updates when temporarily disabled for login
-	if(!bTemporarilyDisabled)
+	if( CurrentSourceControlProvider != nullptr )
 	{
-		// check for any pending dispatches
-		if(PendingStatusUpdateFiles.Num() > 0)
+		ISourceControlProvider& Provider = GetProvider();
+
+		// tick the provider, so any operation results can be read back
+		Provider.Tick();
+
+		// don't allow background status updates when temporarily disabled for login
+		if(!bTemporarilyDisabled)
 		{
-			// grab a batch of files
-			TArray<FString> FilesToDispatch;
-			for(auto Iter(PendingStatusUpdateFiles.CreateConstIterator()); Iter; Iter++)
+			// check for any pending dispatches
+			if(PendingStatusUpdateFiles.Num() > 0)
 			{
-				if(FilesToDispatch.Num() >= SourceControlConstants::MaxStatusDispatchesPerTick)
+				// grab a batch of files
+				TArray<FString> FilesToDispatch;
+				for(auto Iter(PendingStatusUpdateFiles.CreateConstIterator()); Iter; Iter++)
 				{
-					break;
+					if(FilesToDispatch.Num() >= SourceControlConstants::MaxStatusDispatchesPerTick)
+					{
+						break;
+					}
+					FilesToDispatch.Add(*Iter);
 				}
-				FilesToDispatch.Add(*Iter);
-			}
 
-			if(FilesToDispatch.Num() > 0)
-			{
-				// remove the files we are dispatching so we don't try again
-				PendingStatusUpdateFiles.RemoveAt(0, FilesToDispatch.Num());
+				if(FilesToDispatch.Num() > 0)
+				{
+					// remove the files we are dispatching so we don't try again
+					PendingStatusUpdateFiles.RemoveAt(0, FilesToDispatch.Num());
 
-				// dispatch update
-				Provider.Execute(ISourceControlOperation::Create<FUpdateStatus>(), FilesToDispatch, EConcurrency::Asynchronous);
+					// dispatch update
+					Provider.Execute(ISourceControlOperation::Create<FUpdateStatus>(), FilesToDispatch, EConcurrency::Asynchronous);
+				}
 			}
 		}
 	}
