Added delegates to notify you when the source control provider is changed

diff --git a/Engine/Source/Developer/SourceControl/Private/SourceControlModule.cpp b/Engine/Source/Developer/SourceControl/Private/SourceControlModule.cpp
--- a/Engine/Source/Developer/SourceControl/Private/SourceControlModule.cpp
+++ b/Engine/Source/Developer/SourceControl/Private/SourceControlModule.cpp
@@ -325,10 +325,20 @@ void FSourceControlModule::SetProvider( const FName& InName )
 
 void FSourceControlModule::ClearCurrentSourceControlProvider()
 {
-	if( CurrentSourceControlProvider != NULL )
+	if( !CurrentSourceControlProvider || CurrentSourceControlProvider != &DefaultSourceControlProvider )
 	{
-		CurrentSourceControlProvider->Close();
+		ISourceControlProvider* OldSourceControlProvider = CurrentSourceControlProvider;
+		if (CurrentSourceControlProvider)
+		{
+			CurrentSourceControlProvider->Close();
+		}
+
 		CurrentSourceControlProvider = &DefaultSourceControlProvider;
+
+		if (OldSourceControlProvider)
+		{
+			OnSourceControlProviderChanged.Broadcast(*OldSourceControlProvider, *CurrentSourceControlProvider);
+		}
 	}
 }
 
@@ -354,12 +364,19 @@ void FSourceControlModule::SetCurrentSourceControlProvider(ISourceControlProvide
 
 	ClearCurrentSourceControlProvider();
 
+	ISourceControlProvider* OldSourceControlProvider = CurrentSourceControlProvider;
+
 	CurrentSourceControlProvider = &InProvider;
 	CurrentSourceControlProvider->Init(false);	// Don't force a connection here, as its synchronous. Let the user establish a connection.
 
 	SourceControlSettings.SetProvider(CurrentSourceControlProvider->GetName().ToString());
 
 	SaveSettings();
+
+	if (OldSourceControlProvider)
+	{
+		OnSourceControlProviderChanged.Broadcast(*OldSourceControlProvider, *CurrentSourceControlProvider);
+	}
 }
 
 FName FSourceControlModule::GetSourceControlProviderName(int32 ProviderIndex)
@@ -403,6 +420,16 @@ void FSourceControlModule::SetUseGlobalSettings(bool bIsUseGlobalSettings)
 	ShowLoginDialog(FSourceControlLoginClosed(), ELoginWindowMode::Modeless, EOnLoginWindowStartup::PreserveProvider);
 }	
 
+FDelegateHandle FSourceControlModule::RegisterProviderChanged(const FSourceControlProviderChanged::FDelegate& SourceControlProviderChanged)
+{
+	return OnSourceControlProviderChanged.Add(SourceControlProviderChanged);
+}
+
+void FSourceControlModule::UnregisterProviderChanged(FDelegateHandle Handle)
+{
+	OnSourceControlProviderChanged.Remove(Handle);
+}
+
 IMPLEMENT_MODULE( FSourceControlModule, SourceControl );
 
 #undef LOCTEXT_NAMESPACE
diff --git a/Engine/Source/Developer/SourceControl/Private/SourceControlModule.h b/Engine/Source/Developer/SourceControl/Private/SourceControlModule.h
--- a/Engine/Source/Developer/SourceControl/Private/SourceControlModule.h
+++ b/Engine/Source/Developer/SourceControl/Private/SourceControlModule.h
@@ -33,6 +33,8 @@ public:
 	virtual void ShowLoginDialog(const FSourceControlLoginClosed& InOnSourceControlLoginClosed, ELoginWindowMode::Type InLoginWindowMode, EOnLoginWindowStartup::Type InOnLoginWindowStartup = EOnLoginWindowStartup::ResetProviderToNone) override;
 	virtual bool GetUseGlobalSettings() const override;
 	virtual void SetUseGlobalSettings(bool bIsUseGlobalSettings) override;
+	virtual FDelegateHandle RegisterProviderChanged(const FSourceControlProviderChanged::FDelegate& SourceControlProviderChanged) override;
+	virtual void UnregisterProviderChanged(FDelegateHandle Handle) override;
 
 	/** Save the settings to the ini file */
 	void SaveSettings();
@@ -110,4 +112,7 @@ private:
 
 	/** Active Provider name to track source control provider changes */
 	FString ActiveProviderName;
+
+	/** For notifying when the source provider is changed */
+	FSourceControlProviderChanged OnSourceControlProviderChanged;
 };
diff --git a/Engine/Source/Developer/SourceControl/Public/ISourceControlModule.h b/Engine/Source/Developer/SourceControl/Public/ISourceControlModule.h
--- a/Engine/Source/Developer/SourceControl/Public/ISourceControlModule.h
+++ b/Engine/Source/Developer/SourceControl/Public/ISourceControlModule.h
@@ -10,6 +10,9 @@ SOURCECONTROL_API DECLARE_LOG_CATEGORY_EXTERN(LogSourceControl, Log, All);
 /** Delegate called when the source control login window is closed. Parameter determines if source control is enabled or not */
 DECLARE_DELEGATE_OneParam( FSourceControlLoginClosed, bool );
 
+/** Delegate called when the active source control provider is changed */
+DECLARE_MULTICAST_DELEGATE_TwoParams( FSourceControlProviderChanged, ISourceControlProvider& /*OldProvider*/, ISourceControlProvider& /*NewProvider*/ );
+
 /**
  * The modality of the login window.
  */
@@ -110,6 +113,16 @@ public:
 	 */
 	virtual void SetUseGlobalSettings(bool bIsUseGlobalSettings) = 0;
 
+	/**
+	 * Register a delegate to be called when the source control provider changes
+	 */
+	virtual FDelegateHandle RegisterProviderChanged(const FSourceControlProviderChanged::FDelegate& SourceControlProviderChanged) = 0;
+
+	/**
+	 * Unregister a delegate to be called when the source control provider changes
+	 */
+	virtual void UnregisterProviderChanged(FDelegateHandle Handle) = 0;
+
 	/**
 	 * Gets a reference to the source control module instance.
 	 *
