Copying //UE4/Dev-Enterprise to //UE4/Dev-Main (Source: //UE4/Dev-Enterprise @ 3972172)

diff --git a/Engine/Source/Developer/AssetTools/Private/AssetTools.cpp b/Engine/Source/Developer/AssetTools/Private/AssetTools.cpp
--- a/Engine/Source/Developer/AssetTools/Private/AssetTools.cpp
+++ b/Engine/Source/Developer/AssetTools/Private/AssetTools.cpp
@@ -7,14 +7,15 @@
 #include "Misc/FileHelper.h"
 #include "Misc/Paths.h"
 #include "Misc/ScopedSlowTask.h"
+#include "UObject/GCObjectScopeGuard.h"
 #include "UObject/UObjectHash.h"
 #include "UObject/UObjectIterator.h"
 #include "Engine/Blueprint.h"
 #include "Exporters/Exporter.h"
 #include "Editor/EditorEngine.h"
-#include "ISourceControlOperation.h"
 #include "SourceControlOperations.h"
 #include "ISourceControlModule.h"
+#include "SourceControlHelpers.h"
 #include "Editor/UnrealEdEngine.h"
 #include "Settings/EditorLoadingSavingSettings.h"
 #include "ThumbnailRendering/ThumbnailManager.h"
@@ -107,9 +108,11 @@
 #include "Framework/Application/SlateApplication.h"
 #include "Widgets/Notifications/SNotificationList.h"
 #include "AutomatedAssetImportData.h"
+#include "AssetImportTask.h"
 #include "Dialogs/DlgPickPath.h"
 #include "Misc/FeedbackContext.h"
 #include "BusyCursor.h"
+#include "AssetExportTask.h"
 
 #define LOCTEXT_NAMESPACE "AssetTools"
 
@@ -715,6 +718,50 @@ TArray<UObject*> UAssetToolsImpl::ImportAssetsAutomated(const UAutomatedAssetImp
 	return ImportAssetsInternal(ImportData->Filenames, ImportData->DestinationPath, nullptr, Params);
 }
 
+void UAssetToolsImpl::ImportAssetTasks(const TArray<UAssetImportTask*>& ImportTasks) const
+{
+	FScopedSlowTask SlowTask(ImportTasks.Num(), LOCTEXT("ImportSlowTask", "Importing"));
+	SlowTask.MakeDialog();
+
+	FAssetImportParams Params;
+	Params.bSyncToBrowser = false;
+
+	TArray<FString> Filenames;
+	Filenames.Add(TEXT(""));
+	TArray<UPackage*> PackagesToSave;
+	for (UAssetImportTask* ImportTask : ImportTasks)
+	{
+		if (!ImportTask)
+		{
+			UE_LOG(LogAssetTools, Warning, TEXT("ImportAssetTasks() supplied an empty task"));
+			continue;
+		}
+
+		SlowTask.EnterProgressFrame(1, FText::Format(LOCTEXT("Import_ImportingFile", "Importing \"{0}\"..."), FText::FromString(FPaths::GetBaseFilename(ImportTask->Filename))));
+
+		Params.AssetImportTask = ImportTask;
+		Params.bForceOverrideExisting = ImportTask->bReplaceExisting;
+		Params.bAutomated = ImportTask->bAutomated;
+		Filenames[0] = ImportTask->Filename;
+		TArray<UObject*> ImportedObjects = ImportAssetsInternal(Filenames, ImportTask->DestinationPath, nullptr, Params);
+
+		PackagesToSave.Reset(1); 
+		for (UObject* Object : ImportedObjects)
+		{
+			ImportTask->ImportedObjectPaths.Add(Object->GetPathName());
+			if (ImportTask->bSave)
+			{
+				PackagesToSave.AddUnique(Object->GetOutermost());
+			}
+		}
+
+		if (ImportTask->bSave)
+		{
+			UEditorLoadingAndSavingUtils::SavePackages(PackagesToSave, true);
+		}
+	}
+}
+
 void UAssetToolsImpl::ExportAssets(const TArray<FString>& AssetsToExport, const FString& ExportPath) const
 {
 	TArray<UObject*> AssetObjectsToExport;
@@ -1327,6 +1374,11 @@ TArray<UObject*> UAssetToolsImpl::ImportAssetsInternal(const TArray<FString>& Fi
 	bool bReplaceAll = false;
 	bool bDontOverwriteAny = false;
 	bool bDontReplaceAny = false;
+	if (Params.AssetImportTask && Params.AssetImportTask->bAutomated)
+	{
+		bOverwriteAll = bReplaceAll = Params.AssetImportTask->bReplaceExisting;
+		bDontOverwriteAny = bDontReplaceAny = !Params.AssetImportTask->bReplaceExisting;
+	}
 
 	TArray<UFactory*> UsedFactories;
 
@@ -1407,7 +1459,17 @@ TArray<UObject*> UAssetToolsImpl::ImportAssetsInternal(const TArray<FString>& Fi
 			bool bImportWasCancelled = false;
 			FDateTime ImportStartTime = FDateTime::UtcNow();
 
-			FString Name = ObjectTools::SanitizeObjectName(FPaths::GetBaseFilename(Filename));
+			FString Name;
+			if (Params.AssetImportTask && !Params.AssetImportTask->DestinationName.IsEmpty())
+			{
+				Name = Params.AssetImportTask->DestinationName;
+			}
+			else
+			{
+				Name = FPaths::GetBaseFilename(Filename);
+			}
+			Name = ObjectTools::SanitizeObjectName(Name);
+
 			FString PackageName = FPaths::Combine(*DestinationPath, *Name);
 
 			// We can not create assets that share the name of a map file in the same location
@@ -1569,11 +1631,13 @@ TArray<UObject*> UAssetToolsImpl::ImportAssetsInternal(const TArray<FString>& Fi
 			}
 
 			Factory->SetAutomatedAssetImportData(Params.ImportData);
+			Factory->SetAssetImportTask(Params.AssetImportTask);
 
 			ImportAssetType = Factory->ResolveSupportedClass();
 			UObject* Result = Factory->ImportObject(ImportAssetType, Pkg, FName(*Name), RF_Public | RF_Standalone | RF_Transactional, Filename, nullptr, bImportWasCancelled);
 
 			Factory->SetAutomatedAssetImportData(nullptr);
+			Factory->SetAssetImportTask(nullptr);
 
 			// Do not report any error if the operation was canceled.
 			if(!bImportWasCancelled)
@@ -1913,16 +1977,19 @@ void UAssetToolsImpl::ExportAssetsInternal(const TArray<UObject*>& ObjectsToExpo
 					UsedExporters.Add(ExporterToUse);
 				}
 
-				UExporter::FExportToFileParams Params;
-				Params.Object = ObjectToExport;
-				Params.Exporter = ExporterToUse;
-				Params.Filename = *SaveFileName;
-				Params.InSelectedOnly = false;
-				Params.NoReplaceIdentical = false;
-				Params.Prompt = false;
-				Params.bUseFileArchive = ObjectToExport->IsA(UPackage::StaticClass());
-				Params.WriteEmptyFiles = false;
-				UExporter::ExportToFileEx(Params);
+				UAssetExportTask* ExportTask = NewObject<UAssetExportTask>();
+				FGCObjectScopeGuard ExportTaskGuard(ExportTask);
+				ExportTask->Object = ObjectToExport;
+				ExportTask->Exporter = ExporterToUse;
+				ExportTask->Filename = SaveFileName;
+				ExportTask->bSelected = false;
+				ExportTask->bReplaceIdentical = true;
+				ExportTask->bPrompt = false;
+				ExportTask->bUseFileArchive = ObjectToExport->IsA(UPackage::StaticClass());
+				ExportTask->bWriteEmptyFiles = false;
+
+				UExporter::RunAssetExportTask(ExportTask);
+
 				if (ExporterToUse->GetBatchMode() && ExporterToUse->GetCancelBatch())
 				{
 					//Exit the export file loop when there is a cancel all
diff --git a/Engine/Source/Developer/AssetTools/Private/AssetTools.h b/Engine/Source/Developer/AssetTools/Private/AssetTools.h
--- a/Engine/Source/Developer/AssetTools/Private/AssetTools.h
+++ b/Engine/Source/Developer/AssetTools/Private/AssetTools.h
@@ -14,6 +14,7 @@ class FMenuBuilder;
 class IClassTypeActions;
 class UAutomatedAssetImportData;
 class UFactory;
+class UAssetImportTask;
 
 /** Parameters for importing specific set of files */
 struct FAssetImportParams
@@ -21,6 +22,7 @@ struct FAssetImportParams
 	FAssetImportParams()
 		: SpecifiedFactory(nullptr)
 		, ImportData(nullptr)
+		, AssetImportTask(nullptr)
 		, bSyncToBrowser(true)
 		, bForceOverrideExisting(false)
 		, bAutomated(false)
@@ -30,6 +32,8 @@ struct FAssetImportParams
 	UFactory* SpecifiedFactory;
 	/** Data used to determine rules for importing assets through the automated command line interface */
 	const UAutomatedAssetImportData* ImportData;
+	/** Script exposed rules and state for importing assets */
+	const UAssetImportTask* AssetImportTask;
 	/** Whether or not to sync the content browser to the assets after import */
 	bool bSyncToBrowser : 1;
 	/** Whether or not we are forcing existing assets to be overriden without asking */
@@ -78,6 +82,7 @@ public:
 	virtual TArray<UObject*> ImportAssetsWithDialog(const FString& DestinationPath) override;
 	virtual TArray<UObject*> ImportAssets(const TArray<FString>& Files, const FString& DestinationPath, UFactory* ChosenFactory, bool bSyncToBrowser = true, TArray<TPair<FString, FString>>* FilesAndDestinations = nullptr) const override;
 	virtual TArray<UObject*> ImportAssetsAutomated(const UAutomatedAssetImportData* ImportData) const override;
+	virtual void ImportAssetTasks(const TArray<UAssetImportTask*>& ImportTasks) const override;
 	virtual void ExportAssets(const TArray<FString>& AssetsToExport, const FString& ExportPath) const override;
 	virtual void ExportAssets(const TArray<UObject*>& AssetsToExport, const FString& ExportPath) const override;
 	virtual void ExportAssetsWithDialog(const TArray<UObject*>& AssetsToExport, bool bPromptForIndividualFilenames) const override;
diff --git a/Engine/Source/Developer/AssetTools/Public/IAssetTools.h b/Engine/Source/Developer/AssetTools/Public/IAssetTools.h
--- a/Engine/Source/Developer/AssetTools/Public/IAssetTools.h
+++ b/Engine/Source/Developer/AssetTools/Public/IAssetTools.h
@@ -19,6 +19,7 @@ struct FAssetData;
 class IAssetTypeActions;
 class IClassTypeActions;
 class UFactory;
+class UAssetImportTask;
 
 USTRUCT(BlueprintType)
 struct FAssetRenameData
@@ -228,6 +229,15 @@ public:
 	UFUNCTION(BlueprintCallable, Category = "Editor Scripting | Asset Tools")
 	virtual TArray<UObject*> ImportAssetsAutomated( const UAutomatedAssetImportData* ImportData) const = 0;
 
+	/**
+	* Imports assets using tasks specified.
+	*
+	* @param ImportTasks	Tasks that specify how to import each file
+	* @return list of successfully imported assets
+	*/
+	UFUNCTION(BlueprintCallable, Category = "Editor Scripting | Asset Tools")
+	virtual void ImportAssetTasks(const TArray<UAssetImportTask*>& ImportTasks) const = 0;
+
 	/**
 	 * Exports the specified objects to file.
 	 *
diff --git a/Engine/Source/Developer/FunctionalTesting/Private/GroundTruthData.cpp b/Engine/Source/Developer/FunctionalTesting/Private/GroundTruthData.cpp
--- a/Engine/Source/Developer/FunctionalTesting/Private/GroundTruthData.cpp
+++ b/Engine/Source/Developer/FunctionalTesting/Private/GroundTruthData.cpp
@@ -10,7 +10,8 @@
 #if WITH_EDITOR
 #include "ISourceControlModule.h"
 #include "ISourceControlProvider.h"
-#include "ISourceControlOperation.h"
+#include "SourceControlOperations.h"
+#include "SourceControlHelpers.h"
 #endif
 
 DEFINE_LOG_CATEGORY_STATIC(GroundTruthLog, Log, Log)
diff --git a/Engine/Source/Developer/Localization/Private/LocalizationSourceControlUtil.cpp b/Engine/Source/Developer/Localization/Private/LocalizationSourceControlUtil.cpp
--- a/Engine/Source/Developer/Localization/Private/LocalizationSourceControlUtil.cpp
+++ b/Engine/Source/Developer/Localization/Private/LocalizationSourceControlUtil.cpp
@@ -3,8 +3,8 @@
 #include "LocalizationSourceControlUtil.h"
 #include "ISourceControlModule.h"
 #include "ISourceControlProvider.h"
-#include "ISourceControlOperation.h"
 #include "SourceControlOperations.h"
+#include "SourceControlHelpers.h"
 #include "Misc/Paths.h"
 
 DEFINE_LOG_CATEGORY_STATIC(LogLocalizationSourceControl, Log, All);
diff --git a/Engine/Source/Developer/MeshReductionInterface/Public/IMeshReductionInterfaces.h b/Engine/Source/Developer/MeshReductionInterface/Public/IMeshReductionInterfaces.h
--- a/Engine/Source/Developer/MeshReductionInterface/Public/IMeshReductionInterfaces.h
+++ b/Engine/Source/Developer/MeshReductionInterface/Public/IMeshReductionInterfaces.h
@@ -7,6 +7,8 @@
 #include "Engine/MeshMerging.h"
 #include "Modules/ModuleInterface.h"
 
+struct FOverlappingCorners;
+
 /**
 * Mesh reduction interface.
 */
@@ -26,7 +28,7 @@ public:
 		struct FRawMesh& OutReducedMesh,
 		float& OutMaxDeviation,
 		const struct FRawMesh& InMesh,
-		const TMultiMap<int32, int32>& InOverlappingCorners,
+		const FOverlappingCorners& InOverlappingCorners,
 		const struct FMeshReductionSettings& ReductionSettings
 	) = 0;
 	/**
diff --git a/Engine/Source/Developer/MeshUtilities/Private/LayoutUV.cpp b/Engine/Source/Developer/MeshUtilities/Private/LayoutUV.cpp
--- a/Engine/Source/Developer/MeshUtilities/Private/LayoutUV.cpp
+++ b/Engine/Source/Developer/MeshUtilities/Private/LayoutUV.cpp
@@ -23,7 +23,7 @@ FLayoutUV::FLayoutUV( FRawMesh* InMesh, uint32 InSrcChannel, uint32 InDstChannel
 	, NextMeshChartId( 0 )
 {}
 
-int32 FLayoutUV::FindCharts( const TMultiMap<int32,int32>& OverlappingCorners )
+int32 FLayoutUV::FindCharts( const FOverlappingCorners& OverlappingCorners )
 {
 	double Begin = FPlatformTime::Seconds();
 
@@ -43,10 +43,10 @@ int32 FLayoutUV::FindCharts( const TMultiMap<int32,int32>& OverlappingCorners )
 	FDisjointSet DisjointSet( NumTris );
 	for( uint32 i = 0; i < NumIndexes; i++ )
 	{
-		for( auto It = OverlappingCorners.CreateConstKeyIterator(i); It; ++It )
+		const TArray<int32>& Overlapping = OverlappingCorners.FindIfOverlapping(i);
+		for (int32 It : Overlapping)
 		{
-			uint32 j = It.Value();
-
+			uint32 j = It;
 			if( j > i )
 			{
 				const uint32 TriI = i/3;
diff --git a/Engine/Source/Developer/MeshUtilities/Private/MeshUtilities.cpp b/Engine/Source/Developer/MeshUtilities/Private/MeshUtilities.cpp
--- a/Engine/Source/Developer/MeshUtilities/Private/MeshUtilities.cpp
+++ b/Engine/Source/Developer/MeshUtilities/Private/MeshUtilities.cpp
@@ -1535,13 +1535,117 @@ static void ComputeTriangleTangents(
 	check(TriangleTangentZ.Num() == NumTriangles);*/
 }
 
+void FOverlappingCorners::Init(int32 NumIndices)
+{
+	Arrays.Reset();
+	Sets.Reset();
+	bFinishedAdding = false;
+
+	IndexBelongsTo.Reset(NumIndices);
+	IndexBelongsTo.AddUninitialized(NumIndices);
+	FMemory::Memset(IndexBelongsTo.GetData(), 0xFF, NumIndices * sizeof(int32));
+}
+
+void FOverlappingCorners::Add(int32 Key, int32 Value)
+{
+	check(Key != Value);
+	check(bFinishedAdding == false);
+
+	int32 ContainerIndex = IndexBelongsTo[Key];
+	if (ContainerIndex == INDEX_NONE)
+	{
+		ContainerIndex = Arrays.Num();
+		TArray<int32>& Container = Arrays.AddDefaulted_GetRef();
+		Container.Reserve(6);
+		Container.Add(Key);
+		Container.Add(Value);
+		IndexBelongsTo[Key] = ContainerIndex;
+		IndexBelongsTo[Value] = ContainerIndex;
+	}
+	else
+	{
+		IndexBelongsTo[Value] = ContainerIndex;
+
+		TArray<int32>& ArrayContainer = Arrays[ContainerIndex];
+		if (ArrayContainer.Num() == 1)
+		{
+			// Container is a set
+			Sets[ArrayContainer.Last()].Add(Value);
+		}
+		else
+		{
+			// Container is an array
+			ArrayContainer.AddUnique(Value);
+
+			// Change container into set when one vertex is shared by large number of triangles
+			if (ArrayContainer.Num() > 12)
+			{
+				int32 SetIndex = Sets.Num();
+				TSet<int32>& Set = Sets.AddDefaulted_GetRef();
+				Set.Append(ArrayContainer);
+
+				// Having one element means we are using a set
+				// An array will never have just 1 element normally because we add them as pairs
+				ArrayContainer.Reset(1);
+				ArrayContainer.Add(SetIndex);
+			}
+		}
+	}
+}
+
+void FOverlappingCorners::FinishAdding()
+{
+	check(bFinishedAdding == false);
+
+	for (TArray<int32>& Array : Arrays)
+	{
+		// Turn sets back into arrays for easier iteration code
+		// Also reduces peak memory later in the import process
+		if (Array.Num() == 1)
+		{
+			TSet<int32>& Set = Sets[Array.Last()];
+			Array.Reset(Set.Num());
+			for (int32 i : Set)
+			{
+				Array.Add(i);
+			}
+		}
+
+		// Sort arrays now to avoid sort multiple times
+		Array.Sort();
+	}
+
+	Sets.Empty();
+
+	bFinishedAdding = true;
+}
+
+uint32 FOverlappingCorners::GetAllocatedSize(void) const
+{
+	uint32 BaseMemoryAllocated = IndexBelongsTo.GetAllocatedSize() + Arrays.GetAllocatedSize() + Sets.GetAllocatedSize();
+
+	uint32 ArraysMemory = 0;
+	for (const TArray<int32>& ArrayIt : Arrays)
+	{
+		ArraysMemory += ArrayIt.GetAllocatedSize();
+	}
+
+	uint32 SetsMemory = 0;
+	for (const TSet<int32>& SetsIt : Sets)
+	{
+		SetsMemory += SetsIt.GetAllocatedSize();
+	}
+
+	return BaseMemoryAllocated + ArraysMemory + SetsMemory;
+}
+
 /**
 * Create a table that maps the corner of each face to its overlapping corners.
 * @param OutOverlappingCorners - Maps a corner index to the indices of all overlapping corners.
 * @param RawMesh - The mesh for which to compute overlapping corners.
 */
 void FMeshUtilities::FindOverlappingCorners(
-	TMultiMap<int32, int32>& OutOverlappingCorners,
+	FOverlappingCorners& OutOverlappingCorners,
 	const TArray<FVector>& InVertices,
 	const TArray<uint32>& InIndices,
 	float ComparisonThreshold) const
@@ -1559,6 +1663,8 @@ void FMeshUtilities::FindOverlappingCorners(
 	// Sort the vertices by z value
 	VertIndexAndZ.Sort(FCompareIndexAndZ());
 
+	OutOverlappingCorners.Init(NumWedges);
+
 	// Search for duplicates, quickly!
 	for (int32 i = 0; i < VertIndexAndZ.Num(); i++)
 	{
@@ -1574,10 +1680,11 @@ void FMeshUtilities::FindOverlappingCorners(
 			if (PointsEqual(PositionA, PositionB, ComparisonThreshold))
 			{
 				OutOverlappingCorners.Add(VertIndexAndZ[i].Index, VertIndexAndZ[j].Index);
-				OutOverlappingCorners.Add(VertIndexAndZ[j].Index, VertIndexAndZ[i].Index);
 			}
 		}
 	}
+
+	OutOverlappingCorners.FinishAdding();
 }
 
 /**
@@ -1586,7 +1693,7 @@ void FMeshUtilities::FindOverlappingCorners(
 * @param RawMesh - The mesh for which to compute overlapping corners.
 */
 void FMeshUtilities::FindOverlappingCorners(
-	TMultiMap<int32, int32>& OutOverlappingCorners,
+	FOverlappingCorners& OutOverlappingCorners,
 	FRawMesh const& RawMesh,
 	float ComparisonThreshold
 	) const
@@ -1611,7 +1718,7 @@ static void ComputeTangents(
 	const TArray<uint32>& InIndices,
 	const TArray<FVector2D>& InUVs,
 	const TArray<uint32>& SmoothingGroupIndices,
-	TMultiMap<int32, int32> const& OverlappingCorners,
+	const FOverlappingCorners& OverlappingCorners,
 	TArray<FVector>& OutTangentX,
 	TArray<FVector>& OutTangentY,
 	TArray<FVector>& OutTangentZ,
@@ -1640,7 +1747,6 @@ static void ComputeTangents(
 	// Declare these out here to avoid reallocations.
 	TArray<FFanFace> RelevantFacesForCorner[3];
 	TArray<int32> AdjacentFaces;
-	TArray<int32> DupVerts;
 
 	int32 NumWedges = InIndices.Num();
 	int32 NumFaces = NumWedges / 3;
@@ -1712,13 +1818,15 @@ static void ComputeTangents(
 		for (int32 CornerIndex = 0; CornerIndex < 3; CornerIndex++)
 		{
 			int32 ThisCornerIndex = WedgeOffset + CornerIndex;
-			DupVerts.Reset();
-			OverlappingCorners.MultiFind(ThisCornerIndex, DupVerts);
-			DupVerts.Add(ThisCornerIndex); // I am a "dup" of myself
+			const TArray<int32>& DupVerts = OverlappingCorners.FindIfOverlapping(ThisCornerIndex);
 			for (int32 k = 0; k < DupVerts.Num(); k++)
 			{
 				AdjacentFaces.AddUnique(DupVerts[k] / 3);
 			}
+			if (DupVerts.Num() == 0)
+			{
+				AdjacentFaces.AddUnique(ThisCornerIndex / 3); // I am a "dup" of myself
+			}
 		}
 
 		// We need to sort these here because the criteria for point equality is
@@ -1935,7 +2043,7 @@ static void ComputeTangents(
 
 static void ComputeTangents(
 	FRawMesh& RawMesh,
-	TMultiMap<int32, int32> const& OverlappingCorners,
+	const FOverlappingCorners& OverlappingCorners,
 	uint32 TangentOptions
 	)
 {
@@ -2145,7 +2253,7 @@ static void ComputeNormals(
 	const TArray<uint32>& InIndices,
 	const TArray<FVector2D>& InUVs,
 	const TArray<uint32>& SmoothingGroupIndices,
-	TMultiMap<int32, int32> const& OverlappingCorners,
+	const FOverlappingCorners& OverlappingCorners,
 	TArray<FVector>& OutTangentZ,
 	const uint32 TangentOptions
 	)
@@ -2172,7 +2280,6 @@ static void ComputeNormals(
 	// Declare these out here to avoid reallocations.
 	TArray<FFanFace> RelevantFacesForCorner[3];
 	TArray<int32> AdjacentFaces;
-	TArray<int32> DupVerts;
 
 	int32 NumWedges = InIndices.Num();
 	int32 NumFaces = NumWedges / 3;
@@ -2223,9 +2330,11 @@ static void ComputeNormals(
 		for (int32 CornerIndex = 0; CornerIndex < 3; CornerIndex++)
 		{
 			int32 ThisCornerIndex = WedgeOffset + CornerIndex;
-			DupVerts.Reset();
-			OverlappingCorners.MultiFind(ThisCornerIndex, DupVerts);
-			DupVerts.Add(ThisCornerIndex); // I am a "dup" of myself
+			const TArray<int32>& DupVerts = OverlappingCorners.FindIfOverlapping(ThisCornerIndex);
+			if (DupVerts.Num() == 0)
+			{
+				AdjacentFaces.AddUnique(ThisCornerIndex / 3); // I am a "dup" of myself
+			}
 			for (int32 k = 0; k < DupVerts.Num(); k++)
 			{
 				AdjacentFaces.AddUnique(DupVerts[k] / 3);
@@ -2396,7 +2505,7 @@ static void ComputeTangents_MikkTSpace(
 	const TArray<uint32>& InIndices,
 	const TArray<FVector2D>& InUVs,
 	const TArray<uint32>& SmoothingGroupIndices,
-	TMultiMap<int32, int32> const& OverlappingCorners,
+	const FOverlappingCorners& OverlappingCorners,
 	TArray<FVector>& OutTangentX,
 	TArray<FVector>& OutTangentY,
 	TArray<FVector>& OutTangentZ,
@@ -2460,7 +2569,7 @@ static void ComputeTangents_MikkTSpace(
 
 static void ComputeTangents_MikkTSpace(
 	FRawMesh& RawMesh,
-	TMultiMap<int32, int32> const& OverlappingCorners,
+	const FOverlappingCorners& OverlappingCorners,
 	uint32 TangentOptions
 	)
 {
@@ -2608,7 +2717,7 @@ void FMeshUtilities::BuildStaticMeshVertexAndIndexBuffers(
 	TArray<TArray<uint32> >& OutPerSectionIndices,
 	TArray<int32>& OutWedgeMap,
 	const FRawMesh& RawMesh,
-	const TMultiMap<int32, int32>& OverlappingCorners,
+	const FOverlappingCorners& OverlappingCorners,
 	const TMap<uint32, uint32>& MaterialToSectionMapping,
 	float ComparisonThreshold,
 	FVector BuildScale,
@@ -2616,7 +2725,6 @@ void FMeshUtilities::BuildStaticMeshVertexAndIndexBuffers(
 	)
 {
 	TMap<int32, int32> FinalVerts;
-	TArray<int32> DupVerts;
 	int32 NumFaces = RawMesh.WedgeIndices.Num() / 3;
 
 	// Process each face, build vertex buffer and per-section index buffers.
@@ -2647,9 +2755,7 @@ void FMeshUtilities::BuildStaticMeshVertexAndIndexBuffers(
 			int32 WedgeIndex = FaceIndex * 3 + CornerIndex;
 			FStaticMeshBuildVertex ThisVertex = BuildStaticMeshVertex(RawMesh, WedgeIndex, BuildScale);
 
-			DupVerts.Reset();
-			OverlappingCorners.MultiFind(WedgeIndex, DupVerts);
-			DupVerts.Sort();
+			const TArray<int32>& DupVerts = OverlappingCorners.FindIfOverlapping(WedgeIndex);
 
 			int32 Index = INDEX_NONE;
 			for (int32 k = 0; k < DupVerts.Num(); k++)
@@ -2790,7 +2896,7 @@ public:
 		{
 			FStaticMeshSourceModel& SrcModel = SourceModels[LODIndex];
 			FRawMesh& RawMesh = *new(LODMeshes)FRawMesh;
-			TMultiMap<int32, int32>& OverlappingCorners = *new(LODOverlappingCorners)TMultiMap<int32, int32>;
+			FOverlappingCorners& OverlappingCorners = *new(LODOverlappingCorners)FOverlappingCorners;
 
 			if (!SrcModel.RawMeshBulkData->IsEmpty())
 			{
@@ -2940,8 +3046,8 @@ public:
 			{
 				FRawMesh& InMesh = LODMeshes[ReductionSettings.BaseLODModel];
 				FRawMesh& DestMesh = LODMeshes[NumValidLODs];
-				TMultiMap<int32, int32>& InOverlappingCorners = LODOverlappingCorners[ReductionSettings.BaseLODModel];
-				TMultiMap<int32, int32>& DestOverlappingCorners = LODOverlappingCorners[NumValidLODs];
+				FOverlappingCorners& InOverlappingCorners = LODOverlappingCorners[ReductionSettings.BaseLODModel];
+				FOverlappingCorners& DestOverlappingCorners = LODOverlappingCorners[NumValidLODs];
 
 				MeshReduction->Reduce(DestMesh, LODMaxDeviation[NumValidLODs], InMesh, InOverlappingCorners, ReductionSettings);
 				if (DestMesh.WedgeIndices.Num() > 0 && !DestMesh.IsValid())
@@ -2952,7 +3058,6 @@ public:
 				OutWasReduced[LODIndex] = true;
 
 				// Recompute adjacency information.
-				DestOverlappingCorners.Reset();
 				float ComparisonThreshold = GetComparisonThreshold(LODBuildSettings[NumValidLODs]);
 				MeshUtilities.FindOverlappingCorners(DestOverlappingCorners, DestMesh, ComparisonThreshold);
 
@@ -3274,7 +3379,7 @@ private:
 	int32 NumValidLODs;
 
 	TIndirectArray<FRawMesh> LODMeshes;
-	TIndirectArray<TMultiMap<int32, int32> > LODOverlappingCorners;
+	TIndirectArray<FOverlappingCorners> LODOverlappingCorners;
 	float LODMaxDeviation[MAX_STATIC_MESH_LODS];
 	FMeshBuildSettings LODBuildSettings[MAX_STATIC_MESH_LODS];
 	bool HasRawMesh[MAX_STATIC_MESH_LODS];
@@ -3539,7 +3644,7 @@ public:
 
 public:
 	void Skeletal_FindOverlappingCorners(
-		TMultiMap<int32, int32>& OutOverlappingCorners,
+		FOverlappingCorners& OutOverlappingCorners,
 		IMeshBuildData* BuildData,
 		float ComparisonThreshold
 		)
@@ -3563,6 +3668,8 @@ public:
 		// Sort the vertices by z value
 		VertIndexAndZ.Sort(FCompareIndexAndZ());
 
+		OutOverlappingCorners.Init(NumWedges);
+
 		// Search for duplicates, quickly!
 		for (int32 i = 0; i < VertIndexAndZ.Num(); i++)
 		{
@@ -3578,10 +3685,11 @@ public:
 				if (PointsEqual(PositionA, PositionB, ComparisonThreshold))
 				{
 					OutOverlappingCorners.Add(VertIndexAndZ[i].Index, VertIndexAndZ[j].Index);
-					OutOverlappingCorners.Add(VertIndexAndZ[j].Index, VertIndexAndZ[i].Index);
 				}
 			}
 		}
+
+		OutOverlappingCorners.FinishAdding();
 	}
 
 	void Skeletal_ComputeTriangleTangents(
@@ -3642,7 +3750,7 @@ public:
 
 	void Skeletal_ComputeTangents(
 		IMeshBuildData* BuildData,
-		TMultiMap<int32, int32> const& OverlappingCorners
+		const FOverlappingCorners& OverlappingCorners
 		)
 	{
 		bool bBlendOverlappingNormals = true;
@@ -3668,7 +3776,6 @@ public:
 		// Declare these out here to avoid reallocations.
 		TArray<FFanFace> RelevantFacesForCorner[3];
 		TArray<int32> AdjacentFaces;
-		TArray<int32> DupVerts;
 
 		int32 NumFaces = BuildData->GetNumFaces();
 		int32 NumWedges = BuildData->GetNumWedges();
@@ -3741,9 +3848,11 @@ public:
 			for (int32 CornerIndex = 0; CornerIndex < 3; CornerIndex++)
 			{
 				int32 ThisCornerIndex = WedgeOffset + CornerIndex;
-				DupVerts.Reset();
-				OverlappingCorners.MultiFind(ThisCornerIndex, DupVerts);
-				DupVerts.Add(ThisCornerIndex); // I am a "dup" of myself
+				const TArray<int32>& DupVerts = OverlappingCorners.FindIfOverlapping(ThisCornerIndex);
+				if (DupVerts.Num() == 0)
+				{
+					AdjacentFaces.AddUnique(ThisCornerIndex / 3); // I am a "dup" of myself
+				}
 				for (int32 k = 0; k < DupVerts.Num(); k++)
 				{
 					AdjacentFaces.AddUnique(DupVerts[k] / 3);
@@ -3964,7 +4073,7 @@ public:
 
 	void Skeletal_ComputeTangents_MikkTSpace(
 		IMeshBuildData* BuildData,
-		TMultiMap<int32, int32> const& OverlappingCorners
+		const FOverlappingCorners& OverlappingCorners
 		)
 	{
 		bool bBlendOverlappingNormals = true;
@@ -3990,7 +4099,6 @@ public:
 		// Declare these out here to avoid reallocations.
 		TArray<FFanFace> RelevantFacesForCorner[3];
 		TArray<int32> AdjacentFaces;
-		TArray<int32> DupVerts;
 
 		int32 NumFaces = BuildData->GetNumFaces();
 		int32 NumWedges = BuildData->GetNumWedges();
@@ -4055,9 +4163,11 @@ public:
 			for (int32 CornerIndex = 0; CornerIndex < 3; CornerIndex++)
 			{
 				int32 ThisCornerIndex = WedgeOffset + CornerIndex;
-				DupVerts.Reset();
-				OverlappingCorners.MultiFind(ThisCornerIndex, DupVerts);
-				DupVerts.Add(ThisCornerIndex); // I am a "dup" of myself
+				const TArray<int32>& DupVerts = OverlappingCorners.FindIfOverlapping(ThisCornerIndex);
+				if (DupVerts.Num() == 0)
+				{
+					AdjacentFaces.AddUnique(ThisCornerIndex / 3); // I am a "dup" of myself
+				}
 				for (int32 k = 0; k < DupVerts.Num(); k++)
 				{
 					AdjacentFaces.AddUnique(DupVerts[k] / 3);
@@ -4252,7 +4362,7 @@ public:
 
 		BeginSlowTask();
 
-		TMultiMap<int32, int32>& OverlappingCorners = *new(LODOverlappingCorners)TMultiMap<int32, int32>;
+		FOverlappingCorners& OverlappingCorners = *new(LODOverlappingCorners)FOverlappingCorners;
 
 		float ComparisonThreshold = THRESH_POINTS_ARE_SAME;//GetComparisonThreshold(LODBuildSettings[LODIndex]);
 		int32 NumWedges = BuildData->GetNumWedges();
@@ -4518,7 +4628,7 @@ private:
 		GenerateRendering,
 	};
 
-	TIndirectArray<TMultiMap<int32, int32> > LODOverlappingCorners;
+	TIndirectArray<FOverlappingCorners> LODOverlappingCorners;
 	EStage Stage;
 };
 
@@ -4539,12 +4649,31 @@ bool FMeshUtilities::BuildSkeletalMesh(FSkeletalMeshLODModel& LODModel, const FR
 		{
 			FSkelMeshSection& CurSection = InLODModel.Sections[SectionIdx];
 			const int32 NumSoftVertices = CurSection.SoftVertices.Num();
-			for (int32 SrcVertIndex = 0; SrcVertIndex < NumSoftVertices; ++SrcVertIndex)
+
+			// Create a list of vertex Z/index pairs
+			TArray<FIndexAndZ> VertIndexAndZ;
+			VertIndexAndZ.Reserve(NumSoftVertices);
+			for (int32 VertIndex = 0; VertIndex < NumSoftVertices; ++VertIndex)
 			{
+				FSoftSkinVertex& SrcVert = CurSection.SoftVertices[VertIndex];
+				new(VertIndexAndZ)FIndexAndZ(VertIndex, SrcVert.Position);
+			}
+			VertIndexAndZ.Sort(FCompareIndexAndZ());
+
+			// Search for duplicates, quickly!
+			for (int32 i = 0; i < VertIndexAndZ.Num(); ++i)
+			{
+				const uint32 SrcVertIndex = VertIndexAndZ[i].Index;
+				const float Z = VertIndexAndZ[i].Z;
 				FSoftSkinVertex& SrcVert = CurSection.SoftVertices[SrcVertIndex];
 
-				for (int32 IterVertIndex = SrcVertIndex + 1; IterVertIndex < NumSoftVertices; ++IterVertIndex)
+				// only need to search forward, since we add pairs both ways
+				for (int32 j = i + 1; j < VertIndexAndZ.Num(); ++j)
 				{
+					if (FMath::Abs(VertIndexAndZ[j].Z - Z) > THRESH_POINTS_ARE_SAME)
+						break; // can't be any more dups
+
+					const uint32 IterVertIndex = VertIndexAndZ[j].Index;
 					FSoftSkinVertex& IterVert = CurSection.SoftVertices[IterVertIndex];
 					if (PointsEqual(SrcVert.Position, IterVert.Position))
 					{
@@ -5145,14 +5274,14 @@ void FMeshUtilities::RecomputeTangentsAndNormalsForRawMesh(bool bRecomputeTangen
 	if (bRecomputeNormals || bRecomputeTangents)
 	{
 		float ComparisonThreshold = InBuildSettings.bRemoveDegenerates ? THRESH_POINTS_ARE_SAME : 0.0f;
-		TMultiMap<int32, int32> OverlappingCorners;
+		FOverlappingCorners OverlappingCorners;
 		FindOverlappingCorners(OverlappingCorners, OutRawMesh, ComparisonThreshold);
 
 		RecomputeTangentsAndNormalsForRawMesh( bRecomputeTangents, bRecomputeNormals, InBuildSettings, OverlappingCorners, OutRawMesh );
 	}
 }
 
-void FMeshUtilities::RecomputeTangentsAndNormalsForRawMesh(bool bRecomputeTangents, bool bRecomputeNormals, const FMeshBuildSettings& InBuildSettings, const TMultiMap<int32, int32>& InOverlappingCorners, FRawMesh &OutRawMesh) const
+void FMeshUtilities::RecomputeTangentsAndNormalsForRawMesh(bool bRecomputeTangents, bool bRecomputeNormals, const FMeshBuildSettings& InBuildSettings, const FOverlappingCorners& InOverlappingCorners, FRawMesh &OutRawMesh) const
 {
 	const int32 NumWedges = OutRawMesh.WedgeIndices.Num();
 
@@ -5222,7 +5351,7 @@ void FMeshUtilities::ExtractMeshDataForGeometryCache(FRawMesh& RawMesh, const FM
 	}
 
 	// Compute any missing tangents.
-	TMultiMap<int32, int32> OverlappingCorners;
+	FOverlappingCorners OverlappingCorners;
 	if (bRecomputeNormals || bRecomputeTangents)
 	{
 		float ComparisonThreshold = GetComparisonThreshold(BuildSettings);
@@ -5680,7 +5809,8 @@ bool FMeshUtilities::GenerateUniqueUVsForSkeletalMesh(const FSkeletalMeshLODMode
 	}
 
 	// Build overlapping corners map
-	TMultiMap<int32, int32> OverlappingCorners;
+	FOverlappingCorners OverlappingCorners;
+	OverlappingCorners.Init(NumCorners);
 	for (int32 Index = 0; Index < NumCorners; Index++)
 	{
 		int VertexIndex = LODModel.IndexBuffer[Index];
@@ -5692,6 +5822,7 @@ bool FMeshUtilities::GenerateUniqueUVsForSkeletalMesh(const FSkeletalMeshLODMode
 			}
 		}
 	}
+	OverlappingCorners.FinishAdding();
 
 	// Generate new UVs
 	FLayoutUV Packer(&TempMesh, 0, 1, FMath::Clamp(TextureResolution / 4, 32, 512));
@@ -5711,7 +5842,7 @@ void FMeshUtilities::CalculateTangents(const TArray<FVector>& InVertices, const
 {
 	const float ComparisonThreshold = (InTangentOptions & ETangentOptions::IgnoreDegenerateTriangles ) ? THRESH_POINTS_ARE_SAME : 0.0f;
 
-	TMultiMap<int32, int32> OverlappingCorners;
+	FOverlappingCorners OverlappingCorners;
 	FindOverlappingCorners(OverlappingCorners, InVertices, InIndices, ComparisonThreshold);
 
 	if ( InTangentOptions & ETangentOptions::UseMikkTSpace )
@@ -5728,13 +5859,13 @@ void FMeshUtilities::CalculateNormals(const TArray<FVector>& InVertices, const T
 {
 	const float ComparisonThreshold = (InTangentOptions & ETangentOptions::IgnoreDegenerateTriangles ) ? THRESH_POINTS_ARE_SAME : 0.0f;
 
-	TMultiMap<int32, int32> OverlappingCorners;
+	FOverlappingCorners OverlappingCorners;
 	FindOverlappingCorners(OverlappingCorners, InVertices, InIndices, ComparisonThreshold);
 
 	ComputeNormals(InVertices, InIndices, InUVs, InSmoothingGroupIndices, OverlappingCorners, OutNormals, InTangentOptions);
 }
 
-void FMeshUtilities::CalculateOverlappingCorners(const TArray<FVector>& InVertices, const TArray<uint32>& InIndices, bool bIgnoreDegenerateTriangles, TMultiMap<int32, int32>& OutOverlappingCorners) const
+void FMeshUtilities::CalculateOverlappingCorners(const TArray<FVector>& InVertices, const TArray<uint32>& InIndices, bool bIgnoreDegenerateTriangles, FOverlappingCorners& OutOverlappingCorners) const
 {
 	const float ComparisonThreshold = bIgnoreDegenerateTriangles ? THRESH_POINTS_ARE_SAME : 0.f;
 	FindOverlappingCorners(OutOverlappingCorners, InVertices, InIndices, ComparisonThreshold);
@@ -6135,7 +6266,7 @@ bool FMeshUtilities::GenerateUniqueUVsForStaticMesh(const FRawMesh& RawMesh, int
 
 	// Find overlapping corners for UV generator. Allow some threshold - this should not produce any error in a case if resulting
 	// mesh will not merge these vertices.
-	TMultiMap<int32, int32> OverlappingCorners;
+	FOverlappingCorners OverlappingCorners;
 	FModuleManager::Get().LoadModuleChecked<IMeshUtilities>("MeshUtilities").FindOverlappingCorners(OverlappingCorners, RawMesh.VertexPositions, RawMesh.WedgeIndices, THRESH_POINTS_ARE_SAME);
 
 	// Generate new UVs
diff --git a/Engine/Source/Developer/MeshUtilities/Private/MeshUtilitiesPrivate.h b/Engine/Source/Developer/MeshUtilities/Private/MeshUtilitiesPrivate.h
--- a/Engine/Source/Developer/MeshUtilities/Private/MeshUtilitiesPrivate.h
+++ b/Engine/Source/Developer/MeshUtilities/Private/MeshUtilitiesPrivate.h
@@ -81,7 +81,7 @@ private:
 		TArray<TArray<uint32> >& OutPerSectionIndices,
 		TArray<int32>& OutWedgeMap,
 		const FRawMesh& RawMesh,
-		const TMultiMap<int32, int32>& OverlappingCorners,
+		const FOverlappingCorners& OverlappingCorners,
 		const TMap<uint32, uint32>& MaterialToSectionMapping,
 		float ComparisonThreshold,
 		FVector BuildScale,
@@ -101,7 +101,7 @@ private:
 		FDistanceFieldVolumeData& OutData) override;
 
 	virtual void RecomputeTangentsAndNormalsForRawMesh(bool bRecomputeTangents, bool bRecomputeNormals, const FMeshBuildSettings& InBuildSettings, FRawMesh &OutRawMesh) const override;
-	virtual void RecomputeTangentsAndNormalsForRawMesh(bool bRecomputeTangents, bool bRecomputeNormals, const FMeshBuildSettings& InBuildSettings, const TMultiMap<int32, int32>& InOverlappingCorners, FRawMesh &OutRawMesh) const override;
+	virtual void RecomputeTangentsAndNormalsForRawMesh(bool bRecomputeTangents, bool bRecomputeNormals, const FMeshBuildSettings& InBuildSettings, const FOverlappingCorners& InOverlappingCorners, FRawMesh &OutRawMesh) const override;
 
 	virtual bool GenerateUniqueUVsForStaticMesh(const FRawMesh& RawMesh, int32 TextureResolution, TArray<FVector2D>& OutTexCoords) const override;
 
@@ -140,9 +140,9 @@ private:
 	*/
 	void BuildSkeletalModelFromChunks(FSkeletalMeshLODModel& LODModel, const FReferenceSkeleton& RefSkeleton, TArray<FSkinnedMeshChunk*>& Chunks, const TArray<int32>& PointToOriginalMap);
 
-	virtual void FindOverlappingCorners(TMultiMap<int32, int32>& OutOverlappingCorners, const TArray<FVector>& InVertices, const TArray<uint32>& InIndices, float ComparisonThreshold) const override;
+	virtual void FindOverlappingCorners(FOverlappingCorners& OutOverlappingCorners, const TArray<FVector>& InVertices, const TArray<uint32>& InIndices, float ComparisonThreshold) const override;
 
-	void FindOverlappingCorners(TMultiMap<int32, int32>& OutOverlappingCorners, FRawMesh const& RawMesh, float ComparisonThreshold) const;
+	void FindOverlappingCorners(FOverlappingCorners& OutOverlappingCorners, FRawMesh const& RawMesh, float ComparisonThreshold) const;
 	// IModuleInterface interface.
 	virtual void StartupModule() override;
 	virtual void ShutdownModule() override;
@@ -157,7 +157,7 @@ private:
 
 	virtual void CalculateTangents(const TArray<FVector>& InVertices, const TArray<uint32>& InIndices, const TArray<FVector2D>& InUVs, const TArray<uint32>& InSmoothingGroupIndices, const uint32 InTangentOptions, TArray<FVector>& OutTangentX, TArray<FVector>& OutTangentY, TArray<FVector>& OutNormals) const override;
 	virtual void CalculateNormals(const TArray<FVector>& InVertices, const TArray<uint32>& InIndices, const TArray<FVector2D>& InUVs, const TArray<uint32>& InSmoothingGroupIndices, const uint32 InTangentOptions, TArray<FVector>& OutNormals) const override;
-	virtual void CalculateOverlappingCorners(const TArray<FVector>& InVertices, const TArray<uint32>& InIndices, bool bIgnoreDegenerateTriangles, TMultiMap<int32, int32>& OutOverlappingCorners) const override;
+	virtual void CalculateOverlappingCorners(const TArray<FVector>& InVertices, const TArray<uint32>& InIndices, bool bIgnoreDegenerateTriangles, FOverlappingCorners& OutOverlappingCorners) const override;
 
 	// Need to call some members from this class, (which is internal to this module)
 	friend class FStaticMeshUtilityBuilder;
diff --git a/Engine/Source/Developer/MeshUtilities/Public/LayoutUV.h b/Engine/Source/Developer/MeshUtilities/Public/LayoutUV.h
--- a/Engine/Source/Developer/MeshUtilities/Public/LayoutUV.h
+++ b/Engine/Source/Developer/MeshUtilities/Public/LayoutUV.h
@@ -50,7 +50,7 @@ class MESHUTILITIES_API FLayoutUV
 public:
 				FLayoutUV( FRawMesh* InMesh, uint32 InSrcChannel, uint32 InDstChannel, uint32 InTextureResolution );
 
-	int32		FindCharts( const TMultiMap<int32,int32>& OverlappingCorners );
+	int32		FindCharts( const FOverlappingCorners& OverlappingCorners );
 	bool		FindBestPacking();
 	void		CommitPackedUVs();
 
diff --git a/Engine/Source/Developer/MeshUtilities/Public/MeshUtilities.h b/Engine/Source/Developer/MeshUtilities/Public/MeshUtilities.h
--- a/Engine/Source/Developer/MeshUtilities/Public/MeshUtilities.h
+++ b/Engine/Source/Developer/MeshUtilities/Public/MeshUtilities.h
@@ -58,6 +58,45 @@ struct FBoneVertInfo
 	TArray<FVector>	Normals;
 };
 
+/**
+* Container to hold overlapping corners. For a vertex, lists all the overlapping vertices
+*/
+struct FOverlappingCorners
+{
+	FOverlappingCorners() : bFinishedAdding(false)
+	{
+	}
+
+	/* Resets, pre-allocates memory, marks all indices as not overlapping in preperation for calls to Add() */
+	void Init(int32 NumIndices);
+
+	/* Add overlapping indices pair */
+	void Add(int32 Key, int32 Value);
+
+	/* Sorts arrays, converts sets to arrays for sorting and to allow simple iterating code, prevents additional adding */
+	void FinishAdding();
+
+	/* Estimate memory allocated */
+	uint32 GetAllocatedSize(void) const;
+
+	/** 
+	* @return array of sorted overlapping indices including input 'Key', empty array for indices that have no overlaps.
+	*/
+	const TArray<int32>& FindIfOverlapping(int32 Key) const
+	{
+		check(bFinishedAdding);
+		int32 ContainerIndex = IndexBelongsTo[Key];
+		return (ContainerIndex != INDEX_NONE) ? Arrays[ContainerIndex] : EmptyArray;
+	}
+
+private:
+	TArray<int32> IndexBelongsTo;
+	TArray< TArray<int32> > Arrays;
+	TArray< TSet<int32> > Sets;
+	TArray<int32> EmptyArray;
+	bool bFinishedAdding;
+};
+
 class IMeshUtilities : public IModuleInterface
 {
 public:
@@ -176,7 +215,7 @@ public:
 		TArray<TArray<uint32> >& OutPerSectionIndices,
 		TArray<int32>& OutWedgeMap,
 		const FRawMesh& RawMesh,
-		const TMultiMap<int32, int32>& OverlappingCorners,
+		const FOverlappingCorners& OverlappingCorners,
 		const TMap<uint32, uint32>& MaterialToSectionMapping,
 		float ComparisonThreshold,
 		FVector BuildScale,
@@ -332,12 +371,12 @@ public:
 	 * @param InVertices Vertices that make up the mesh
 	 * @param InIndices Indices for the Vertex array
 	 * @param bIgnoreDegenerateTriangles Indicates if we should skip degenerate triangles
-	 * @param OutOverlappingCorners MultiMap to hold the overlapping corners. For a vertex, lists all the overlapping vertices.
+	 * @param OutOverlappingCorners Container to hold the overlapping corners. For a vertex, lists all the overlapping vertices.
 	 */
-	virtual void CalculateOverlappingCorners(const TArray<FVector>& InVertices, const TArray<uint32>& InIndices, bool bIgnoreDegenerateTriangles, TMultiMap<int32, int32>& OutOverlappingCorners) const = 0;
+	virtual void CalculateOverlappingCorners(const TArray<FVector>& InVertices, const TArray<uint32>& InIndices, bool bIgnoreDegenerateTriangles, FOverlappingCorners& OutOverlappingCorners) const = 0;
 
 	virtual void RecomputeTangentsAndNormalsForRawMesh(bool bRecomputeTangents, bool bRecomputeNormals, const FMeshBuildSettings& InBuildSettings, FRawMesh &OutRawMesh) const = 0;
-	virtual void RecomputeTangentsAndNormalsForRawMesh(bool bRecomputeTangents, bool bRecomputeNormals, const FMeshBuildSettings& InBuildSettings, const TMultiMap<int32, int32>& InOverlappingCorners, FRawMesh &OutRawMesh) const = 0;
+	virtual void RecomputeTangentsAndNormalsForRawMesh(bool bRecomputeTangents, bool bRecomputeNormals, const FMeshBuildSettings& InBuildSettings, const FOverlappingCorners& InOverlappingCorners, FRawMesh &OutRawMesh) const = 0;
 
-	virtual void FindOverlappingCorners(TMultiMap<int32, int32>& OutOverlappingCorners, const TArray<FVector>& InVertices, const TArray<uint32>& InIndices, float ComparisonThreshold) const = 0;
+	virtual void FindOverlappingCorners(FOverlappingCorners& OutOverlappingCorners, const TArray<FVector>& InVertices, const TArray<uint32>& InIndices, float ComparisonThreshold) const = 0;
 };
diff --git a/Engine/Source/Developer/MessageLog/Private/Model/MessageLogListingModel.cpp b/Engine/Source/Developer/MessageLog/Private/Model/MessageLogListingModel.cpp
--- a/Engine/Source/Developer/MessageLog/Private/Model/MessageLogListingModel.cpp
+++ b/Engine/Source/Developer/MessageLog/Private/Model/MessageLogListingModel.cpp
@@ -13,7 +13,7 @@ MessageContainer::TConstIterator FMessageLogListingModel::GetMessageIterator(con
 
 const TSharedPtr< FTokenizedMessage >  FMessageLogListingModel::GetMessageAtIndex( const uint32 PageIndex, const int32 MessageIndex ) const
 {
-	TSharedPtr<FTokenizedMessage> FoundMessage = NULL;
+	TSharedPtr<FTokenizedMessage> FoundMessage = nullptr;
 	FPage* Page = PageAtIndex(PageIndex);
 	if( Page->Messages.IsValidIndex( MessageIndex ) )
 	{
@@ -37,7 +37,7 @@ const TSharedPtr< FTokenizedMessage > FMessageLogListingModel::GetMessageFromDat
 			}
 		}
 	}
-	return NULL;
+	return nullptr;
 }
 
 FText FMessageLogListingModel::GetAllMessagesAsText( const uint32 PageIndex ) const
@@ -144,7 +144,7 @@ FMessageLogListingModel::FPage& FMessageLogListingModel::CurrentPage() const
 FMessageLogListingModel::FPage* FMessageLogListingModel::PageAtIndex(const uint32 PageIndex) const
 {
 	check(PageIndex < (uint32)Pages.Num());
-	if(CachedPage != NULL && CachedPageIndex == PageIndex)
+	if(CachedPage != nullptr && CachedPageIndex == PageIndex)
 	{
 		return CachedPage;
 	}
@@ -164,7 +164,7 @@ FMessageLogListingModel::FPage* FMessageLogListingModel::PageAtIndex(const uint3
 	}
 
 	check(false);	// Should never get here!
-	return NULL;
+	return nullptr;
 }
 
 const FText& FMessageLogListingModel::GetPageTitle( const uint32 PageIndex ) const
@@ -187,7 +187,7 @@ void FMessageLogListingModel::CreateNewPageIfRequired()
 			Pages.AddHead(FPage(PendingPageName));
 
 			// invalidate cache as all indices will change
-			CachedPage = NULL;
+			CachedPage = nullptr;
 		}
 		else
 		{
@@ -204,7 +204,7 @@ bool FMessageLogListingModel::AreMessagesEqual(const TSharedRef< FTokenizedMessa
 	{
 		auto TokenItA(MessageA->GetMessageTokens().CreateConstIterator());
 		auto TokenItB(MessageB->GetMessageTokens().CreateConstIterator());
-		for( ; TokenItA && TokenItB; TokenItA++, TokenItB++)
+		for( ; TokenItA && TokenItB; ++TokenItA, ++TokenItB)
 		{
 			TSharedRef<IMessageToken> TokenA = *TokenItA;
 			TSharedRef<IMessageToken> TokenB = *TokenItB;
@@ -220,22 +220,38 @@ bool FMessageLogListingModel::AreMessagesEqual(const TSharedRef< FTokenizedMessa
 	return false;
 }
 
+// Specialized KeyFunc for TSharedRef<FTokenizedMessage> usage in TSet
+struct FTokenizedMessageKeyFunc
+{
+	typedef typename TCallTraits<TSharedRef<FTokenizedMessage>>::ParamType KeyInitType;
+	typedef typename TCallTraits<TSharedRef<FTokenizedMessage>>::ParamType ElementInitType;
+
+	enum { bAllowDuplicateKeys = false };
+
+	static FORCEINLINE KeyInitType GetSetKey(ElementInitType Element)
+	{
+		return Element;
+	}
+
+	static FORCEINLINE bool Matches(KeyInitType A, KeyInitType B)
+	{
+		return FMessageLogListingModel::AreMessagesEqual(A, B);
+	}
+
+	static FORCEINLINE uint32 GetKeyHash(KeyInitType Key)
+	{
+		return GetTypeHash(Key->ToText().ToString());
+	}
+};
+
 void FMessageLogListingModel::RemoveDuplicates(uint32 PageIndex)
 {
 	FPage* Page = PageAtIndex(PageIndex);
-	if(Page != NULL)
+	
+	if (Page != nullptr && Page->Messages.Num() > 0)
 	{
-		for(int IndexOuter = Page->Messages.Num() - 1; IndexOuter >= 0; --IndexOuter)
-		{
-			for(int IndexInner = Page->Messages.Num() - 1; IndexInner >= 0; --IndexInner)
-			{
-				if(IndexInner != IndexOuter && AreMessagesEqual(Page->Messages[IndexInner], Page->Messages[IndexOuter]))
-				{
-					Page->Messages.RemoveAt(IndexOuter);
-					break;
-				}
-			}
-		}
+		TSet<TSharedRef<FTokenizedMessage>, FTokenizedMessageKeyFunc> MessageSet(Page->Messages);
+		Page->Messages = MessageSet.Array();
 	}
 }
 
diff --git a/Engine/Source/Developer/MessageLog/Private/Model/MessageLogListingModel.h b/Engine/Source/Developer/MessageLog/Private/Model/MessageLogListingModel.h
--- a/Engine/Source/Developer/MessageLog/Private/Model/MessageLogListingModel.h
+++ b/Engine/Source/Developer/MessageLog/Private/Model/MessageLogListingModel.h
@@ -97,7 +97,7 @@ public:
 	/** Helper function for RemoveDuplicates(), exposed so the ViewModel can use it too */
 	static bool AreMessagesEqual(const TSharedRef< FTokenizedMessage >& Message0, const TSharedRef< FTokenizedMessage >& Message1);
 
-	/** Remove any messages that are duplicates of one another - O(n^2) */
+	/** Remove any messages that are duplicates of one another - O(n) */
 	void RemoveDuplicates(uint32 PageIndex);
 
 protected:
diff --git a/Engine/Source/Developer/OutputLog/Private/SOutputLog.cpp b/Engine/Source/Developer/OutputLog/Private/SOutputLog.cpp
--- a/Engine/Source/Developer/OutputLog/Private/SOutputLog.cpp
+++ b/Engine/Source/Developer/OutputLog/Private/SOutputLog.cpp
@@ -11,6 +11,7 @@
 #include "Framework/Commands/UIAction.h"
 #include "Framework/Text/SlateTextLayout.h"
 #include "Framework/Text/SlateTextRun.h"
+#include "Framework/Application/SlateApplication.h"
 #include "Widgets/Text/STextBlock.h"
 #include "Widgets/Input/SMenuAnchor.h"
 #include "Framework/MultiBox/MultiBoxBuilder.h"
@@ -127,7 +128,7 @@ bool FConsoleCommandExecutor::Exec(const TCHAR* Input)
 			}
 		}
 
-		if (!bWasHandled && !Player)
+		if (!bWasHandled)
 		{
 			if (GIsEditor)
 			{
@@ -186,7 +187,7 @@ SConsoleInputBox::SConsoleInputBox()
 }
 
 BEGIN_SLATE_FUNCTION_BUILD_OPTIMIZATION
-void SConsoleInputBox::Construct( const FArguments& InArgs )
+void SConsoleInputBox::Construct(const FArguments& InArgs)
 {
 	OnConsoleCommandExecuted = InArgs._OnConsoleCommandExecuted;
 	ConsoleCommandCustomExec = InArgs._ConsoleCommandCustomExec;
@@ -253,6 +254,8 @@ void SConsoleInputBox::Construct( const FArguments& InArgs )
 			[
 				SNew(SBox)
 				.HeightOverride(250) // avoids flickering, ideally this would be adaptive to the content without flickering
+				.MinDesiredWidth(300)
+				.MaxDesiredWidth(this, &SConsoleInputBox::GetSelectionListMaxWidth)
 				[
 					SAssignNew(SuggestionListView, SListView< TSharedPtr<FString> >)
 					.ListItemsSource(&Suggestions.SuggestionsList)
@@ -309,6 +312,13 @@ void SConsoleInputBox::SuggestionSelectionChanged(TSharedPtr<FString> NewValue,
 //	FSlateApplication::Get().SetKeyboardFocus( WidgetToFocusPath, EFocusCause::SetDirectly );
 }
 
+FOptionalSize SConsoleInputBox::GetSelectionListMaxWidth() const
+{
+	// Limit the width of the suggestions list to the work area that this widget currently resides on
+	const FSlateRect WidgetRect(GetCachedGeometry().GetAbsolutePosition(), GetCachedGeometry().GetAbsolutePosition() + GetCachedGeometry().GetAbsoluteSize());
+	const FSlateRect WidgetWorkArea = FSlateApplication::Get().GetWorkArea(WidgetRect);
+	return FMath::Max(300.0f, WidgetWorkArea.GetSize().X - 12.0f);
+}
 
 TSharedRef<ITableRow> SConsoleInputBox::MakeSuggestionListItemWidget(TSharedPtr<FString> Text, const TSharedRef<STableViewBase>& OwnerTable)
 {
@@ -322,14 +332,10 @@ TSharedRef<ITableRow> SConsoleInputBox::MakeSuggestionListItemWidget(TSharedPtr<
 	return
 		SNew(STableRow< TSharedPtr<FString> >, OwnerTable)
 		[
-			SNew(SBox)
-			.MinDesiredWidth(300)
-			[
-				SNew(STextBlock)
-				.Text(FText::FromString(SanitizedText))
-				.TextStyle(FEditorStyle::Get(), "Log.Normal")
-				.HighlightText(Suggestions.SuggestionsHighlight)
-			]
+			SNew(STextBlock)
+			.Text(FText::FromString(SanitizedText))
+			.TextStyle(FEditorStyle::Get(), "Log.Normal")
+			.HighlightText(Suggestions.SuggestionsHighlight)
 		];
 }
 
diff --git a/Engine/Source/Developer/OutputLog/Private/SOutputLog.h b/Engine/Source/Developer/OutputLog/Private/SOutputLog.h
--- a/Engine/Source/Developer/OutputLog/Private/SOutputLog.h
+++ b/Engine/Source/Developer/OutputLog/Private/SOutputLog.h
@@ -126,6 +126,9 @@ protected:
 
 	void OnTextChanged(const FText& InText);
 
+	/** Get the maximum width of the selection list */
+	FOptionalSize GetSelectionListMaxWidth() const;
+
 	/** Makes the widget for the suggestions messages in the list view */
 	TSharedRef<ITableRow> MakeSuggestionListItemWidget(TSharedPtr<FString> Message, const TSharedRef<STableViewBase>& OwnerTable);
 
diff --git a/Engine/Source/Developer/RawMesh/Private/RawMesh.cpp b/Engine/Source/Developer/RawMesh/Private/RawMesh.cpp
--- a/Engine/Source/Developer/RawMesh/Private/RawMesh.cpp
+++ b/Engine/Source/Developer/RawMesh/Private/RawMesh.cpp
@@ -285,4 +285,9 @@ void FRawMeshBulkData::UseHashAsGuid(UObject* Owner)
 	bGuidIsHash = true;
 }
 
+const FByteBulkData& FRawMeshBulkData::GetBulkData() const
+{
+	return BulkData;
+}
+
 #endif // #if WITH_EDITORONLY_DATA
diff --git a/Engine/Source/Developer/RawMesh/Public/RawMesh.h b/Engine/Source/Developer/RawMesh/Public/RawMesh.h
--- a/Engine/Source/Developer/RawMesh/Public/RawMesh.h
+++ b/Engine/Source/Developer/RawMesh/Public/RawMesh.h
@@ -118,6 +118,8 @@ public:
 	/** Uses a hash as the GUID, useful to prevent creating new GUIDs on load for legacy assets. */
 	RAWMESH_API void UseHashAsGuid(class UObject* Owner);
 
+	RAWMESH_API const FByteBulkData& GetBulkData() const;
+
 	/** Returns true if no bulk data is available for this mesh. */
 	FORCEINLINE bool IsEmpty() const { return BulkData.GetBulkDataSize() == 0; }
 #endif // #if WITH_EDITORONLY_DATA
diff --git a/Engine/Source/Developer/SettingsEditor/Private/Widgets/SSettingsEditor.cpp b/Engine/Source/Developer/SettingsEditor/Private/Widgets/SSettingsEditor.cpp
--- a/Engine/Source/Developer/SettingsEditor/Private/Widgets/SSettingsEditor.cpp
+++ b/Engine/Source/Developer/SettingsEditor/Private/Widgets/SSettingsEditor.cpp
@@ -161,11 +161,10 @@ void SSettingsEditor::NotifyPostChange( const FPropertyChangedEvent& PropertyCha
 		
 		if (PropertyChangedEvent.GetNumObjectsBeingEdited() > 0)
 		{
-			const UObject* ObjectBeingEdited = PropertyChangedEvent.GetObjectBeingEdited(0);
+			UObject* ObjectBeingEdited = (UObject*)PropertyChangedEvent.GetObjectBeingEdited(0);
 
 			// Get the section from the edited object.  We cannot use the selected section as multiple sections can be shown at once in the settings details panel.
 			ISettingsSectionPtr Section = Model->GetSectionFromSectionObject(ObjectBeingEdited);
-			if(Section.IsValid())
 			{
 				FString RelativePath;
 				bool bIsSourceControlled = false;
@@ -194,7 +193,10 @@ void SSettingsEditor::NotifyPostChange( const FPropertyChangedEvent& PropertyCha
 					SettingsHelpers::MakeWritable(FullPath);
 				}
 			
-				RecordPreferenceChangedAnalytics(Section, PropertyChangedEvent);
+				if (Section.IsValid())
+				{
+					RecordPreferenceChangedAnalytics(Section, PropertyChangedEvent);
+				}
 
 				// Determine if the Property is an Array or Array Element
 				bool bIsArrayOrArrayElement = PropertyThatChanged->GetActiveMemberNode()->GetValue()->IsA(UArrayProperty::StaticClass()) 
@@ -207,11 +209,11 @@ void SSettingsEditor::NotifyPostChange( const FPropertyChangedEvent& PropertyCha
 				bool bIsMapOrMapElement = PropertyThatChanged->GetActiveMemberNode()->GetValue()->IsA(UMapProperty::StaticClass())
 					|| ((Outer != nullptr) && Outer->IsA(UMapProperty::StaticClass()));
 
-				if (Section->GetSettingsObject()->GetClass()->HasAnyClassFlags(CLASS_DefaultConfig) && !bIsArrayOrArrayElement && !bIsSetOrSetElement && !bIsMapOrMapElement)
+				if (ObjectBeingEdited->GetClass()->HasAnyClassFlags(CLASS_DefaultConfig) && !bIsArrayOrArrayElement && !bIsSetOrSetElement && !bIsMapOrMapElement)
 				{
-					Section->GetSettingsObject()->UpdateSinglePropertyInConfigFile(PropertyThatChanged->GetActiveMemberNode()->GetValue(), Section->GetSettingsObject()->GetDefaultConfigFilename());
+					ObjectBeingEdited->UpdateSinglePropertyInConfigFile(PropertyThatChanged->GetActiveMemberNode()->GetValue(), ObjectBeingEdited->GetDefaultConfigFilename());
 				}
-				else
+				else if (Section.IsValid())
 				{
 					Section->Save();
 				}
diff --git a/Engine/Source/Developer/SourceControl/Private/ISourceControlProvider.cpp b/Engine/Source/Developer/SourceControl/Private/ISourceControlProvider.cpp
new file mode 100644
--- /dev/null
+++ b/Engine/Source/Developer/SourceControl/Private/ISourceControlProvider.cpp
@@ -0,0 +1,76 @@
+// Copyright 1998-2018 Epic Games, Inc. All Rights Reserved.
+
+#include "ISourceControlProvider.h"
+#include "SourceControlHelpers.h"
+#include "SourceControlOperations.h"
+
+
+#define LOCTEXT_NAMESPACE "ISourceControlProvider"
+
+
+ECommandResult::Type ISourceControlProvider::Login(const FString& InPassword, EConcurrency::Type InConcurrency, const FSourceControlOperationComplete& InOperationCompleteDelegate)
+{
+	TSharedRef<FConnect, ESPMode::ThreadSafe> ConnectOperation = ISourceControlOperation::Create<FConnect>();
+	ConnectOperation->SetPassword(InPassword);
+	return Execute(ConnectOperation, InConcurrency, InOperationCompleteDelegate);
+}
+
+ECommandResult::Type ISourceControlProvider::GetState(const TArray<UPackage*>& InPackages, TArray< TSharedRef<ISourceControlState, ESPMode::ThreadSafe> >& OutState, EStateCacheUsage::Type InStateCacheUsage)
+{
+	TArray<FString> Files = SourceControlHelpers::PackageFilenames(InPackages);
+	return GetState(Files, OutState, InStateCacheUsage);
+}
+
+TSharedPtr<ISourceControlState, ESPMode::ThreadSafe> ISourceControlProvider::GetState(const UPackage* InPackage, EStateCacheUsage::Type InStateCacheUsage)
+{
+	return GetState(SourceControlHelpers::PackageFilename(InPackage), InStateCacheUsage);
+}
+
+TSharedPtr<ISourceControlState, ESPMode::ThreadSafe> ISourceControlProvider::GetState(const FString& InFile, EStateCacheUsage::Type InStateCacheUsage)
+{
+	TArray<FString> Files;
+	TArray< TSharedRef<ISourceControlState, ESPMode::ThreadSafe> > States;
+	Files.Add(InFile);
+	if (GetState(Files, States, InStateCacheUsage) == ECommandResult::Succeeded)
+	{
+		TSharedRef<ISourceControlState, ESPMode::ThreadSafe> State = States[0];
+		return State;
+	}
+	return NULL;
+}
+
+ECommandResult::Type ISourceControlProvider::Execute(const TSharedRef<ISourceControlOperation, ESPMode::ThreadSafe>& InOperation, const EConcurrency::Type InConcurrency, const FSourceControlOperationComplete& InOperationCompleteDelegate)
+{
+	return Execute(InOperation, TArray<FString>(), InConcurrency, InOperationCompleteDelegate);
+}
+
+ECommandResult::Type ISourceControlProvider::Execute(const TSharedRef<ISourceControlOperation, ESPMode::ThreadSafe>& InOperation, const UPackage* InPackage, const EConcurrency::Type InConcurrency, const FSourceControlOperationComplete& InOperationCompleteDelegate)
+{
+	return Execute(InOperation, SourceControlHelpers::PackageFilename(InPackage), InConcurrency, InOperationCompleteDelegate);
+}
+
+ECommandResult::Type ISourceControlProvider::Execute(const TSharedRef<ISourceControlOperation, ESPMode::ThreadSafe>& InOperation, const FString& InFile, const EConcurrency::Type InConcurrency, const FSourceControlOperationComplete& InOperationCompleteDelegate)
+{
+	TArray<FString> FileArray;
+	FileArray.Add(InFile);
+	return Execute(InOperation, FileArray, InConcurrency, InOperationCompleteDelegate);
+}
+
+ECommandResult::Type ISourceControlProvider::Execute(const TSharedRef<ISourceControlOperation, ESPMode::ThreadSafe>& InOperation, const TArray<UPackage*>& InPackages, const EConcurrency::Type InConcurrency, const FSourceControlOperationComplete& InOperationCompleteDelegate)
+{
+	TArray<FString> FileArray = SourceControlHelpers::PackageFilenames(InPackages);
+	return Execute(InOperation, FileArray, InConcurrency, InOperationCompleteDelegate);
+}
+
+TSharedPtr<class ISourceControlLabel> ISourceControlProvider::GetLabel(const FString& InLabelName) const
+{
+	TArray< TSharedRef<class ISourceControlLabel> > Labels = GetLabels(InLabelName);
+	if (Labels.Num() > 0)
+	{
+		return Labels[0];
+	}
+
+	return NULL;
+}
+
+#undef LOCTEXT_NAMESPACE
diff --git a/Engine/Source/Developer/SourceControl/Private/SourceControlHelpers.cpp b/Engine/Source/Developer/SourceControl/Private/SourceControlHelpers.cpp
--- a/Engine/Source/Developer/SourceControl/Private/SourceControlHelpers.cpp
+++ b/Engine/Source/Developer/SourceControl/Private/SourceControlHelpers.cpp
@@ -16,48 +16,877 @@
 
 #define LOCTEXT_NAMESPACE "SourceControlHelpers"
 
-const FString& USourceControlHelpers::GetSettingsIni()
+
+namespace SourceControlHelpersInternal
+{
+
+/*
+ * Status info set by LogError() and USourceControlHelpers methods if an error occurs
+ * regardless whether their bSilent is set or not.
+ * Should be empty if there is was no error.
+ * @see	USourceControlHelpers::LastErrorMsg(), LogError()
+ */
+FText LastErrorText;
+
+/* Store error and write to Log if bSilent is false. */
+inline void LogError(const FText& ErrorText, bool bSilent)
+{
+	LastErrorText = ErrorText;
+
+	if (!bSilent)
+	{
+		FMessageLog("SourceControl").Error(LastErrorText);
+	}
+}
+
+/* Return provider if ready to go, else return nullptr. */
+ISourceControlProvider* VerifySourceControl(bool bSilent)
+{
+	ISourceControlModule& SCModule = ISourceControlModule::Get();
+
+	if (!SCModule.IsEnabled())
+	{
+		LogError(LOCTEXT("SourceControlDisabled", "Source control is not enabled."), bSilent);
+
+		return nullptr;
+	}
+
+	ISourceControlProvider* Provider = &SCModule.GetProvider();
+
+	if (!Provider->IsAvailable())
+	{
+		LogError(LOCTEXT("SourceControlServerUnavailable", "Source control server is currently not available."), bSilent);
+
+		return nullptr;
+	}
+
+	// Clear the last error text if there hasn't been an error (yet).
+	LastErrorText = FText::GetEmpty();
+
+	return Provider;
+}
+
+
+/*
+ * Converts specified file to fully qualified file path that is compatible with source control.
+ *
+ * @param	InFile		File string - can be either fully qualified path, relative path, long package name, asset path or export text path (often stored on clipboard)
+ * @param	bSilent		if false then write out any error info to the Log. Any error text can be retrieved by LastErrorMsg() regardless.
+ * @return	Fully qualified file path to use with source control or "" if conversion unsuccessful.
+ */
+FString ConvertFileToQualifiedPath(const FString& InFile, bool bSilent, const FString& AssociatedExtension = FString())
+{
+	// Converted to qualified file path
+	FString SCFile;
+
+	if (InFile.IsEmpty())
+	{
+		LogError(LOCTEXT("UnspecifiedFile", "File not specified"), bSilent);
+
+		return SCFile;
+	}
+
+	// Try to determine if file is one of:
+	// - fully qualified path
+	// - relative path
+	// - long package name
+	// - asset path
+	// - export text path (often stored on clipboard)
+	//
+	// For example:
+	// - D:\Epic\Dev-Ent\Projects\Python3rdBP\Content\Mannequin\Animations\ThirdPersonIdle.uasset
+	// - Content\Mannequin\Animations\ThirdPersonIdle.uasset
+	// - /Game/Mannequin/Animations/ThirdPersonIdle
+	// - /Game/Mannequin/Animations/ThirdPersonIdle.ThirdPersonIdle
+	// - AnimSequence'/Game/Mannequin/Animations/ThirdPersonIdle.ThirdPersonIdle'
+
+	SCFile = InFile;
+	bool bPackage = false;
+
+
+	// Is ExportTextPath (often stored in Clipboard) form?
+	//  - i.e. AnimSequence'/Game/Mannequin/Animations/ThirdPersonIdle.ThirdPersonIdle'
+	if (SCFile[SCFile.Len() - 1] == '\'')
+	{
+		SCFile = FPackageName::ExportTextPathToObjectPath(SCFile);
+	}
+
+	if (SCFile[0] == '/')
+	{
+		// Assume it is a package
+		bPackage = true;
+
+		// Try to get filename by finding it on disk
+		if (!FPackageName::DoesPackageExist(SCFile, nullptr, &SCFile))
+		{
+			// The package does not exist on disk, see if we can find it in memory and predict the file extension
+			// Only do this if the supplied package name is valid
+			const bool bIncludeReadOnlyRoots = false;
+			bPackage = FPackageName::IsValidLongPackageName(SCFile, bIncludeReadOnlyRoots);
+
+			if (bPackage)
+			{
+				const FString* PackageExtension = &FPackageName::GetAssetPackageExtension();
+
+				if (AssociatedExtension.IsEmpty())
+				{
+					UPackage* Package = FindPackage(nullptr, *SCFile);
+
+					if (Package)
+					{
+						// This is a package in memory that has not yet been saved. Determine the extension and convert to a filename
+						PackageExtension = Package->ContainsMap() ? &FPackageName::GetMapPackageExtension() : &FPackageName::GetAssetPackageExtension();
+					}
+				}
+				else
+				{
+					PackageExtension = &AssociatedExtension;
+				}
+
+				bPackage = FPackageName::TryConvertLongPackageNameToFilename(SCFile, SCFile, *PackageExtension);
+			}
+		}
+
+		if (bPackage)
+		{
+			SCFile = FPaths::ConvertRelativePathToFull(SCFile);
+
+			return SCFile;
+		}
+	}
+
+	// Assume it is a qualified or relative file path
+
+	// Could normalize it
+	//FPaths::NormalizeFilename(SCFile);
+
+	if (!FPaths::IsRelative(SCFile))
+	{
+		return SCFile;
+	}
+
+	// Qualify based on process base directory.
+	// Something akin to "C:/Epic/UE4/Engine/Binaries/Win64/" as a current path.
+	SCFile = FPaths::ConvertRelativePathToFull(InFile);
+
+	if (FPaths::FileExists(SCFile))
+	{
+		return SCFile;
+	}
+
+	// Qualify based on project directory.
+	SCFile = FPaths::ConvertRelativePathToFull(FPaths::ConvertRelativePathToFull(FPaths::ProjectDir()), InFile);
+
+	if (FPaths::FileExists(SCFile))
+	{
+		return SCFile;
+	}
+
+	// Qualify based on Engine directory
+	SCFile = FPaths::ConvertRelativePathToFull(FPaths::ConvertRelativePathToFull(FPaths::EngineDir()), InFile);
+
+	return SCFile;
+}
+
+
+/**
+ * Converts specified files to fully qualified file paths that are compatible with source control.
+ *
+ * @param	InFiles			File strings - can be either fully qualified path, relative path, long package name, asset name or export text path (often stored on clipboard)
+ * @param	OutFilePaths	Fully qualified file paths to use with source control or "" if conversion unsuccessful.
+ * @param	bSilent			if false then write out any error info to the Log. Any error text can be retrieved by LastErrorMsg() regardless.
+ * @return	true if all files successfully converted, false if any had errors
+ */
+bool ConvertFilesToQualifiedPaths(const TArray<FString>& InFiles, TArray<FString>& OutFilePaths, bool bSilent)
+{
+	uint32 SkipNum = 0u;
+
+	for (const FString& File : InFiles)
+	{
+		FString SCFile = ConvertFileToQualifiedPath(File, bSilent);
+
+		if (SCFile.IsEmpty())
+		{
+			SkipNum++;
+		}
+		else
+		{
+			OutFilePaths.Add(MoveTemp(SCFile));
+		}
+	}
+
+	if (SkipNum)
+	{
+		FFormatNamedArguments Arguments;
+		Arguments.Add(TEXT("SkipNum"), FText::AsNumber(SkipNum));
+		LogError(FText::Format(LOCTEXT("FilesSkipped", "During conversion to qualified file paths, {SkipNum} files were skipped!"), Arguments), bSilent);
+
+		return false;
+	}
+
+	return true;
+}
+
+}  // namespace SourceControlHelpersInternal
+
+
+FString USourceControlHelpers::CurrentProvider()
+{
+	// Note that if there is no provider there is still a dummy default provider object
+	ISourceControlProvider& Provider = ISourceControlModule::Get().GetProvider();
+
+	return Provider.GetName().ToString();
+}
+
+
+bool USourceControlHelpers::IsEnabled()
+{
+	return ISourceControlModule::Get().IsEnabled();
+}
+
+
+bool USourceControlHelpers::IsAvailable()
+{
+	ISourceControlModule& SCModule = ISourceControlModule::Get();
+
+	return SCModule.IsEnabled() && SCModule.GetProvider().IsAvailable();
+}
+
+
+FText USourceControlHelpers::LastErrorMsg()
+{
+	return SourceControlHelpersInternal::LastErrorText;
+}
+
+
+bool USourceControlHelpers::CheckOutFile(const FString& InFile, bool bSilent)
+{
+	// Determine file type and ensure it is in form source control wants
+	FString SCFile = SourceControlHelpersInternal::ConvertFileToQualifiedPath(InFile, bSilent);
+
+	if (SCFile.IsEmpty())
+	{
+		return false;
+	}
+
+	// Ensure source control system is up and running
+	ISourceControlProvider* Provider = SourceControlHelpersInternal::VerifySourceControl(bSilent);
+
+	if (!Provider)
+	{
+		return false;
+	}
+
+	FSourceControlStatePtr SCState = Provider->GetState(SCFile, EStateCacheUsage::ForceUpdate);
+
+	if (!SCState.IsValid())
+	{
+		// Improper or invalid SCC state
+		FFormatNamedArguments Arguments;
+		Arguments.Add(TEXT("InFile"), FText::FromString(InFile));
+		Arguments.Add(TEXT("SCFile"), FText::FromString(SCFile));
+		SourceControlHelpersInternal::LogError(FText::Format(LOCTEXT("CouldNotDetermineState", "Could not determine source control state of file '{InFile}' ({SCFile})."), Arguments), bSilent);
+
+		return false;
+	}
+
+	if (SCState->IsCheckedOut() || SCState->IsAdded())
+	{
+		// Already checked out or opened for add
+		return true;
+	}
+
+	bool bCheckOutFail = false;
+
+	if (SCState->CanCheckout())
+	{
+		if (Provider->Execute(ISourceControlOperation::Create<FCheckOut>(), SCFile) == ECommandResult::Succeeded)
+		{
+			return true;
+		}
+
+		bCheckOutFail = true;
+	}
+
+	// Only error info after this point
+
+	FString SimultaneousCheckoutUser;
+	FFormatNamedArguments Arguments;
+	Arguments.Add(TEXT("InFile"), FText::FromString(InFile));
+	Arguments.Add(TEXT("SCFile"), FText::FromString(SCFile));
+
+	if (bCheckOutFail)
+	{
+		SourceControlHelpersInternal::LogError(FText::Format(LOCTEXT("CheckoutFailed", "Failed to check out file '{InFile}' ({SCFile})."), Arguments), bSilent);
+	}
+	else if (!SCState->IsSourceControlled())
+	{
+		SourceControlHelpersInternal::LogError(FText::Format(LOCTEXT("NotSourceControlled", "Could not check out the file '{InFile}' because it is not under source control ({SCFile})."), Arguments), bSilent);
+	}
+	else if (!SCState->IsCurrent())
+	{
+		SourceControlHelpersInternal::LogError(FText::Format(LOCTEXT("NotAtHeadRevision", "File '{InFile}' is not at head revision ({SCFile})."), Arguments), bSilent);
+	}
+	else if (SCState->IsCheckedOutOther(&(SimultaneousCheckoutUser)))
+	{
+		Arguments.Add(TEXT("SimultaneousCheckoutUser"), FText::FromString(SimultaneousCheckoutUser));
+		SourceControlHelpersInternal::LogError(FText::Format(LOCTEXT("SimultaneousCheckout", "File '{InFile}' is checked out by another ('{SimultaneousCheckoutUser}') ({SCFile})."), Arguments), bSilent);
+	}
+	else
+	{
+		// Improper or invalid SCC state
+		SourceControlHelpersInternal::LogError(FText::Format(LOCTEXT("CouldNotDetermineState", "Could not determine source control state of file '{InFile}' ({SCFile})."), Arguments), bSilent);
+	}
+
+	return false;
+}
+
+
+bool USourceControlHelpers::CheckOutFiles(const TArray<FString>& InFiles, bool bSilent)
+{
+	ISourceControlProvider* Provider = SourceControlHelpersInternal::VerifySourceControl(bSilent);
+
+	if (!Provider)
+	{
+		return false;
+	}
+
+	TArray<FString> FilePaths;
+
+	// Even if some files were skipped, still apply to the others
+	bool bFilesSkipped = !SourceControlHelpersInternal::ConvertFilesToQualifiedPaths(InFiles, FilePaths, bSilent);
+
+	// Less error checking and info is made for multiple files than the single file version.
+	// This multi-file version could be made similarly more sophisticated.
+	ECommandResult::Type Result = Provider->Execute(ISourceControlOperation::Create<FCheckOut>(), FilePaths);
+
+	return !bFilesSkipped && (Result == ECommandResult::Succeeded);
+}
+
+
+bool USourceControlHelpers::CheckOutOrAddFile(const FString& InFile, bool bSilent)
+{
+	// Determine file type and ensure it is in form source control wants
+	FString SCFile = SourceControlHelpersInternal::ConvertFileToQualifiedPath(InFile, bSilent);
+
+	if (SCFile.IsEmpty())
+	{
+		return false;
+	}
+
+	// Ensure source control system is up and running
+	ISourceControlProvider* Provider = SourceControlHelpersInternal::VerifySourceControl(bSilent);
+
+	if (!Provider)
+	{
+		return false;
+	}
+
+	FSourceControlStatePtr SCState = Provider->GetState(SCFile, EStateCacheUsage::ForceUpdate);
+
+	if (!SCState.IsValid())
+	{
+		// Improper or invalid SCC state
+		FFormatNamedArguments Arguments;
+		Arguments.Add(TEXT("InFile"), FText::FromString(InFile));
+		Arguments.Add(TEXT("SCFile"), FText::FromString(SCFile));
+		SourceControlHelpersInternal::LogError(FText::Format(LOCTEXT("CouldNotDetermineState", "Could not determine source control state of file '{InFile}' ({SCFile})."), Arguments), bSilent);
+
+		return false;
+	}
+
+	if (SCState->IsCheckedOut() || SCState->IsAdded())
+	{
+		// Already checked out or opened for add
+		return true;
+	}
+
+	// Stuff single file in array for functions that require array
+	TArray<FString> FilesToBeCheckedOut;
+	FilesToBeCheckedOut.Add(SCFile);
+
+	if (SCState->CanCheckout())
+	{
+		if (Provider->Execute(ISourceControlOperation::Create<FCheckOut>(), FilesToBeCheckedOut) != ECommandResult::Succeeded)
+		{
+			FFormatNamedArguments Arguments;
+			Arguments.Add(TEXT("InFile"), FText::FromString(InFile));
+			Arguments.Add(TEXT("SCFile"), FText::FromString(SCFile));
+			SourceControlHelpersInternal::LogError(FText::Format(LOCTEXT("CheckoutFailed", "Failed to check out file '{InFile}' ({SCFile})."), Arguments), bSilent);
+
+			return false;
+		}
+
+		return true;
+	}
+
+	bool bAddFail = false;
+
+	if (!SCState->IsSourceControlled())
+	{
+		if (Provider->Execute(ISourceControlOperation::Create<FMarkForAdd>(), FilesToBeCheckedOut) == ECommandResult::Succeeded)
+		{
+			return true;
+		}
+
+		bAddFail = true;;
+	}
+
+	FString SimultaneousCheckoutUser;
+	FFormatNamedArguments Arguments;
+	Arguments.Add(TEXT("InFile"), FText::FromString(InFile));
+	Arguments.Add(TEXT("SCFile"), FText::FromString(SCFile));
+
+	if (bAddFail)
+	{
+		SourceControlHelpersInternal::LogError(FText::Format(LOCTEXT("AddFailed", "Failed to add file '{InFile}' to source control ({SCFile})."), Arguments), bSilent);
+	}
+	else if (!SCState->IsCurrent())
+	{
+		SourceControlHelpersInternal::LogError(FText::Format(LOCTEXT("NotAtHeadRevision", "File '{InFile}' is not at head revision ({SCFile})."), Arguments), bSilent);
+	}
+	else if (SCState->IsCheckedOutOther(&(SimultaneousCheckoutUser)))
+	{
+		Arguments.Add(TEXT("SimultaneousCheckoutUser"), FText::FromString(SimultaneousCheckoutUser));
+		SourceControlHelpersInternal::LogError(FText::Format(LOCTEXT("SimultaneousCheckout", "File '{InFile}' is checked out by another ({SimultaneousCheckoutUser}) ({SCFile})."), Arguments), bSilent);
+	}
+	else
+	{
+		// Improper or invalid SCC state
+		SourceControlHelpersInternal::LogError(FText::Format(LOCTEXT("CouldNotDetermineState", "Could not determine source control state of file '{InFile}' ({SCFile})."), Arguments), bSilent);
+	}
+
+	return false;
+}
+
+
+bool USourceControlHelpers::MarkFileForAdd(const FString& InFile, bool bSilent)
+{
+	// Determine file type and ensure it is in form source control wants
+	FString SCFile = SourceControlHelpersInternal::ConvertFileToQualifiedPath(InFile, bSilent);
+
+	if (SCFile.IsEmpty())
+	{
+		return false;
+	}
+
+	// Ensure source control system is up and running
+	ISourceControlProvider* Provider = SourceControlHelpersInternal::VerifySourceControl(bSilent);
+
+	if (!Provider)
+	{
+		return false;
+	}
+
+	// Mark for add now if needed
+	FSourceControlStatePtr SCState = Provider->GetState(SCFile, EStateCacheUsage::Use);
+
+	if (!SCState.IsValid())
+	{
+		// Improper or invalid SCC state
+		FFormatNamedArguments Arguments;
+		Arguments.Add(TEXT("InFile"), FText::FromString(InFile));
+		Arguments.Add(TEXT("SCFile"), FText::FromString(SCFile));
+		SourceControlHelpersInternal::LogError(FText::Format(LOCTEXT("CouldNotDetermineState", "Could not determine source control state of file '{InFile}' ({SCFile})."), Arguments), bSilent);
+
+		return false;
+	}
+
+	// Add if necessary
+	if (SCState->IsUnknown() || (!SCState->IsSourceControlled() && !SCState->IsAdded()))
+	{
+		if (Provider->Execute(ISourceControlOperation::Create<FMarkForAdd>(), SCFile) != ECommandResult::Succeeded)
+		{
+			FFormatNamedArguments Arguments;
+			Arguments.Add(TEXT("InFile"), FText::FromString(InFile));
+			Arguments.Add(TEXT("SCFile"), FText::FromString(SCFile));
+			SourceControlHelpersInternal::LogError(FText::Format(LOCTEXT("MarkForAddFailed", "Failed to add file '{InFile}' to source control ({SCFile})."), Arguments), bSilent);
+
+			return false;
+		}
+	}
+
+	return true;
+}
+
+
+bool USourceControlHelpers::MarkFilesForAdd(const TArray<FString>& InFiles, bool bSilent)
+{
+	ISourceControlProvider* Provider = SourceControlHelpersInternal::VerifySourceControl(bSilent);
+
+	if (!Provider)
+	{
+		return false;
+	}
+
+	TArray<FString> FilePaths;
+
+	// Even if some files were skipped, still apply to the others
+	bool bFilesSkipped = !SourceControlHelpersInternal::ConvertFilesToQualifiedPaths(InFiles, FilePaths, bSilent);
+
+	// Less error checking and info is made for multiple files than the single file version.
+	// This multi-file version could be made similarly more sophisticated.
+	ECommandResult::Type Result = Provider->Execute(ISourceControlOperation::Create<FMarkForAdd>(), FilePaths);
+
+	return !bFilesSkipped && (Result == ECommandResult::Succeeded);
+}
+
+
+bool USourceControlHelpers::MarkFileForDelete(const FString& InFile, bool bSilent)
+{
+	// Determine file type and ensure it is in form source control wants
+	FString SCFile = SourceControlHelpersInternal::ConvertFileToQualifiedPath(InFile, bSilent);
+
+	if (SCFile.IsEmpty())
+	{
+		return false;
+	}
+
+	// Ensure source control system is up and running
+	ISourceControlProvider* Provider = SourceControlHelpersInternal::VerifySourceControl(bSilent);
+
+	if (!Provider)
+	{
+		// Error or can't communicate with source control
+		// Could erase it anyway, though keeping it for now.
+		return false;
+	}
+
+	FSourceControlStatePtr SCState = Provider->GetState(SCFile, EStateCacheUsage::ForceUpdate);
+
+	if (!SCState.IsValid())
+	{
+		// Improper or invalid SCC state
+		FFormatNamedArguments Arguments;
+		Arguments.Add(TEXT("InFile"), FText::FromString(InFile));
+		Arguments.Add(TEXT("SCFile"), FText::FromString(SCFile));
+		SourceControlHelpersInternal::LogError(FText::Format(LOCTEXT("CouldNotDetermineState", "Could not determine source control state of file '{InFile}' ({SCFile})."), Arguments), bSilent);
+
+		return false;
+	}
+
+	bool bDelete = false;
+
+	if (SCState->IsSourceControlled())
+	{
+		bool bAdded = SCState->IsAdded();
+
+		if (bAdded || SCState->IsCheckedOut())
+		{
+			if (Provider->Execute(ISourceControlOperation::Create<FRevert>(), SCFile) != ECommandResult::Succeeded)
+			{
+				FFormatNamedArguments Arguments;
+				Arguments.Add(TEXT("InFile"), FText::FromString(InFile));
+				Arguments.Add(TEXT("SCFile"), FText::FromString(SCFile));
+				SourceControlHelpersInternal::LogError(FText::Format(LOCTEXT("CouldNotRevert", "Could not revert source control state of file '{InFile}' ({SCFile})."), Arguments), bSilent);
+
+				return false;
+			}
+		}
+
+		if (!bAdded)
+		{
+			// Was previously added to source control so mark it for delete
+			if (Provider->Execute(ISourceControlOperation::Create<FDelete>(), SCFile) != ECommandResult::Succeeded)
+			{
+				FFormatNamedArguments Arguments;
+				Arguments.Add(TEXT("InFile"), FText::FromString(InFile));
+				Arguments.Add(TEXT("SCFile"), FText::FromString(SCFile));
+				SourceControlHelpersInternal::LogError(FText::Format(LOCTEXT("CouldNotDelete", "Could not delete file '{InFile}' from source control ({SCFile})."), Arguments), bSilent);
+
+				return false;
+			}
+		}
+	}
+
+	// Delete file if it still exists
+	IFileManager& FileManager = IFileManager::Get();
+
+	if (FileManager.FileExists(*SCFile))
+	{
+		// Just a regular file not tracked by source control so erase it.
+		// Don't bother checking if it exists since Delete doesn't care.
+		return FileManager.Delete(*SCFile, false, true);
+	}
+
+	return false;
+}
+
+
+bool USourceControlHelpers::RevertFile(const FString& InFile, bool bSilent)
 {
-	if(ISourceControlModule::Get().GetUseGlobalSettings())
+	// Determine file type and ensure it is in form source control wants
+	FString SCFile = SourceControlHelpersInternal::ConvertFileToQualifiedPath(InFile, bSilent);
+
+	if (SCFile.IsEmpty())
 	{
-		return GetGlobalSettingsIni();
+		return false;
 	}
-	else
+
+	// Ensure source control system is up and running
+	ISourceControlProvider* Provider = SourceControlHelpersInternal::VerifySourceControl(bSilent);
+
+	if (!Provider)
 	{
-		static FString SourceControlSettingsIni;
-		if(SourceControlSettingsIni.Len() == 0)
-		{
-			const FString SourceControlSettingsDir = FPaths::GeneratedConfigDir();
-			FConfigCacheIni::LoadGlobalIniFile(SourceControlSettingsIni, TEXT("SourceControlSettings"), NULL, false, false, true, *SourceControlSettingsDir);
-		}
-		return SourceControlSettingsIni;
+		return false;
 	}
+
+	// Revert file regardless of whether it has had any changes made
+	ECommandResult::Type Result = Provider->Execute(ISourceControlOperation::Create<FRevert>(), SCFile);
+
+	return Result == ECommandResult::Succeeded;
 }
 
-const FString& USourceControlHelpers::GetGlobalSettingsIni()
+
+bool USourceControlHelpers::RevertFiles(const TArray<FString>& InFiles,	bool bSilent)
 {
-	static FString SourceControlGlobalSettingsIni;
-	if(SourceControlGlobalSettingsIni.Len() == 0)
+	ISourceControlProvider* Provider = SourceControlHelpersInternal::VerifySourceControl(bSilent);
+
+	if (!Provider)
 	{
-		const FString SourceControlSettingsDir = FPaths::EngineSavedDir() + TEXT("Config/");
-		FConfigCacheIni::LoadGlobalIniFile(SourceControlGlobalSettingsIni, TEXT("SourceControlSettings"), NULL, false, false, true, *SourceControlSettingsDir);
+		return false;
 	}
-	return SourceControlGlobalSettingsIni;
+
+	// Less error checking and info is made for multiple files than the single file version.
+	// This multi-file version could be made similarly more sophisticated.
+
+	// Revert files regardless of whether they've had any changes made
+	ECommandResult::Type Result = Provider->Execute(ISourceControlOperation::Create<FRevert>(), InFiles);
+
+	return Result == ECommandResult::Succeeded;
+}
+
+
+bool USourceControlHelpers::RevertUnchangedFile(const FString& InFile, bool bSilent)
+{
+	// Determine file type and ensure it is in form source control wants
+	FString SCFile = SourceControlHelpersInternal::ConvertFileToQualifiedPath(InFile, bSilent);
+
+	if (SCFile.IsEmpty())
+	{
+		return false;
+	}
+
+	// Ensure source control system is up and running
+	ISourceControlProvider* Provider = SourceControlHelpersInternal::VerifySourceControl(bSilent);
+
+	if (!Provider)
+	{
+		return false;
+	}
+
+	// Only revert file if they haven't had any changes made
+
+	// Stuff single file in array for functions that require array
+	TArray<FString> InFiles;
+	InFiles.Add(SCFile);
+
+	RevertUnchangedFiles(*Provider, InFiles);
+
+	// Assume it succeeded
+	return true;
+}
+
+
+bool USourceControlHelpers::RevertUnchangedFiles(const TArray<FString>& InFiles, bool bSilent)
+{
+	ISourceControlProvider* Provider = SourceControlHelpersInternal::VerifySourceControl(bSilent);
+
+	if (!Provider)
+	{
+		return false;
+	}
+
+	// Less error checking and info is made for multiple files than the single file version.
+	// This multi-file version could be made similarly more sophisticated.
+
+	// Only revert files if they haven't had any changes made
+	RevertUnchangedFiles(*Provider, InFiles);
+
+	// Assume it succeeded
+	return true;
+}
+
+
+bool USourceControlHelpers::CheckInFile(const FString& InFile, const FString& InDescription, bool bSilent)
+{
+	// Determine file type and ensure it is in form source control wants
+	FString SCFile = SourceControlHelpersInternal::ConvertFileToQualifiedPath(InFile, bSilent);
+
+	if (SCFile.IsEmpty())
+	{
+		return false;
+	}
+
+	// Ensure source control system is up and running
+	ISourceControlProvider* Provider = SourceControlHelpersInternal::VerifySourceControl(bSilent);
+
+	if (!Provider)
+	{
+		return false;
+	}
+
+	TSharedRef<FCheckIn, ESPMode::ThreadSafe> CheckInOp = ISourceControlOperation::Create<FCheckIn>();
+	CheckInOp->SetDescription(FText::FromString(InDescription));
+
+	ECommandResult::Type Result = Provider->Execute(CheckInOp, SCFile);
+
+	return Result == ECommandResult::Succeeded;
+}
+
+
+bool USourceControlHelpers::CheckInFiles(const TArray<FString>& InFiles, const FString& InDescription, bool bSilent)
+{
+	ISourceControlProvider* Provider = SourceControlHelpersInternal::VerifySourceControl(bSilent);
+
+	if (!Provider)
+	{
+		return false;
+	}
+
+	TArray<FString> FilePaths;
+
+	// Even if some files were skipped, still apply to the others
+	bool bFilesSkipped = !SourceControlHelpersInternal::ConvertFilesToQualifiedPaths(InFiles, FilePaths, bSilent);
+
+	// Less error checking and info is made for multiple files than the single file version.
+	// This multi-file version could be made similarly more sophisticated.
+	TSharedRef<FCheckIn, ESPMode::ThreadSafe> CheckInOp = ISourceControlOperation::Create<FCheckIn>();
+	CheckInOp->SetDescription(FText::FromString(InDescription));
+
+	ECommandResult::Type Result = Provider->Execute(CheckInOp, FilePaths);
+
+	return !bFilesSkipped && (Result == ECommandResult::Succeeded);
+}
+
+
+bool USourceControlHelpers::CopyFile(const FString& InSourcePath, const FString& InDestPath, bool bSilent)
+{
+	// Determine file type and ensure it is in form source control wants
+	FString SCSource = SourceControlHelpersInternal::ConvertFileToQualifiedPath(InSourcePath, bSilent);
+
+	if (SCSource.IsEmpty())
+	{
+		return false;
+	}
+
+	// Determine file type and ensure it is in form source control wants
+	FString SCSourcExt(FPaths::GetExtension(SCSource, true));
+	FString SCDest(SourceControlHelpersInternal::ConvertFileToQualifiedPath(InDestPath, bSilent, SCSourcExt));
+
+	if (SCDest.IsEmpty())
+	{
+		return false;
+	}
+
+	// Ensure source control system is up and running
+	ISourceControlProvider* Provider = SourceControlHelpersInternal::VerifySourceControl(bSilent);
+
+	if (!Provider)
+	{
+		return false;
+	}
+
+	TSharedRef<FCopy, ESPMode::ThreadSafe> CopyOp = ISourceControlOperation::Create<FCopy>();
+	CopyOp->SetDestination(SCDest);
+
+	ECommandResult::Type Result = Provider->Execute(CopyOp, SCSource);
+
+	return Result == ECommandResult::Succeeded;
+}
+
+
+FSourceControlState USourceControlHelpers::QueryFileState(const FString& InFile, bool bSilent)
+{
+	FSourceControlState State;
+
+	State.bIsValid = false;
+
+	// Determine file type and ensure it is in form source control wants
+	FString SCFile = SourceControlHelpersInternal::ConvertFileToQualifiedPath(InFile, bSilent);
+
+	if (SCFile.IsEmpty())
+	{
+		State.Filename = InFile;
+		return State;
+	}
+
+	State.Filename = SCFile;
+
+	// Ensure source control system is up and running
+	ISourceControlProvider* Provider = SourceControlHelpersInternal::VerifySourceControl(bSilent);
+
+	if (!Provider)
+	{
+		return State;
+	}
+
+	// Make sure we update the modified state of the files (Perforce requires this
+	// since can be a more expensive test).
+	TSharedRef<FUpdateStatus, ESPMode::ThreadSafe> UpdateStatusOperation = ISourceControlOperation::Create<FUpdateStatus>();
+	UpdateStatusOperation->SetUpdateModifiedState(true);
+	Provider->Execute(UpdateStatusOperation, SCFile);
+
+	FSourceControlStatePtr SCState = Provider->GetState(SCFile, EStateCacheUsage::Use);
+
+	if (!SCState.IsValid())
+	{
+		// Improper or invalid SCC state
+		FFormatNamedArguments Arguments;
+		Arguments.Add(TEXT("InFile"), FText::FromString(InFile));
+		Arguments.Add(TEXT("SCFile"), FText::FromString(SCFile));
+		SourceControlHelpersInternal::LogError(FText::Format(LOCTEXT("CouldNotDetermineState", "Could not determine source control state of file '{InFile}' ({SCFile})."), Arguments), bSilent);
+
+		return State;
+	}
+
+	// Return FSourceControlState rather than a ISourceControlState directly so that
+	// scripting systems can access it.
+
+	State.bIsValid = true;
+
+	// Copy over state info
+	// - make these assignments a method of FSourceControlState if anything else sets a state
+	State.bIsUnknown			= SCState->IsUnknown();
+	State.bIsSourceControlled	= SCState->IsSourceControlled();
+	State.bCanCheckIn			= SCState->CanCheckIn();
+	State.bCanCheckOut			= SCState->CanCheckout();
+	State.bIsCheckedOut			= SCState->IsCheckedOut();
+	State.bIsCurrent			= SCState->IsCurrent();
+	State.bIsAdded				= SCState->IsAdded();
+	State.bIsDeleted			= SCState->IsDeleted();
+	State.bIsIgnored			= SCState->IsIgnored();
+	State.bCanEdit				= SCState->CanEdit();
+	State.bCanDelete			= SCState->CanDelete();
+	State.bCanAdd				= SCState->CanAdd();
+	State.bIsConflicted			= SCState->IsConflicted();
+	State.bCanRevert			= SCState->CanRevert();
+	State.bIsModified			= SCState->IsModified();
+	State.bIsCheckedOutOther	= SCState->IsCheckedOutOther();
+
+	if (State.bIsCheckedOutOther)
+	{
+		SCState->IsCheckedOutOther(&State.CheckedOutOther);
+	}
+
+	return State;
 }
 
+
+
 static FString PackageFilename_Internal( const FString& InPackageName )
 {
 	FString Filename = InPackageName;
 
 	// Get the filename by finding it on disk first
-	if ( !FPackageName::DoesPackageExist(InPackageName, NULL, &Filename) )
+	if ( !FPackageName::DoesPackageExist(InPackageName, nullptr, &Filename) )
 	{
 		// The package does not exist on disk, see if we can find it in memory and predict the file extension
 		// Only do this if the supplied package name is valid
 		const bool bIncludeReadOnlyRoots = false;
 		if ( FPackageName::IsValidLongPackageName(InPackageName, bIncludeReadOnlyRoots) )
 		{
-			UPackage* Package = FindPackage(NULL, *InPackageName);
+			UPackage* Package = FindPackage(nullptr, *InPackageName);
 			if ( Package )
 			{
 				// This is a package in memory that has not yet been saved. Determine the extension and convert to a filename
@@ -70,21 +899,24 @@ static FString PackageFilename_Internal( const FString& InPackageName )
 	return Filename;
 }
 
+
 FString USourceControlHelpers::PackageFilename( const FString& InPackageName )
 {
 	return FPaths::ConvertRelativePathToFull(PackageFilename_Internal(InPackageName));
 }
 
+
 FString USourceControlHelpers::PackageFilename( const UPackage* InPackage )
 {
 	FString Filename;
-	if(InPackage != NULL)
+	if(InPackage != nullptr)
 	{
 		Filename = FPaths::ConvertRelativePathToFull(PackageFilename_Internal(InPackage->GetName()));
 	}
 	return Filename;
 }
 
+
 TArray<FString> USourceControlHelpers::PackageFilenames( const TArray<UPackage*>& InPackages )
 {
 	TArray<FString> OutNames;
@@ -96,6 +928,7 @@ TArray<FString> USourceControlHelpers::PackageFilenames( const TArray<UPackage*>
 	return OutNames;
 }
 
+
 TArray<FString> USourceControlHelpers::PackageFilenames( const TArray<FString>& InPackageNames )
 {
 	TArray<FString> OutNames;
@@ -107,10 +940,12 @@ TArray<FString> USourceControlHelpers::PackageFilenames( const TArray<FString>&
 	return OutNames;
 }
 
+
 TArray<FString> USourceControlHelpers::AbsoluteFilenames( const TArray<FString>& InFileNames )
 {
 	TArray<FString> AbsoluteFiles;
-	for(const auto& FileName : InFileNames)
+
+	for (const FString& FileName : InFileNames)
 	{
 		if(!FPaths::IsRelative(FileName))
 		{
@@ -127,6 +962,7 @@ TArray<FString> USourceControlHelpers::AbsoluteFilenames( const TArray<FString>&
 	return AbsoluteFiles;
 }
 
+
 void USourceControlHelpers::RevertUnchangedFiles( ISourceControlProvider& InProvider, const TArray<FString>& InFiles )
 {
 	// Make sure we update the modified state of the files
@@ -153,6 +989,7 @@ void USourceControlHelpers::RevertUnchangedFiles( ISourceControlProvider& InProv
 	}
 }
 
+
 bool USourceControlHelpers::AnnotateFile( ISourceControlProvider& InProvider, const FString& InLabel, const FString& InFile, TArray<FAnnotationLine>& OutLines )
 {
 	TArray< TSharedRef<ISourceControlLabel> > Labels = InProvider.GetLabels( InLabel );
@@ -201,143 +1038,6 @@ bool USourceControlHelpers::AnnotateFile( ISourceControlProvider& InProvider, in
 	return false;
 }
 
-bool USourceControlHelpers::MarkFileForAdd( const FString& InFilePath )
-{
-	if (InFilePath.IsEmpty())
-	{
-		FMessageLog("SourceControl").Error(LOCTEXT("UnspecifiedCheckoutFile", "Check out file not specified"));
-		return false;
-	}
-
-	if (!ISourceControlModule::Get().IsEnabled())
-	{
-		FMessageLog("SourceControl").Error(LOCTEXT("SourceControlDisabled", "Source control is not enabled."));
-		return false;
-	}
-
-	if (!ISourceControlModule::Get().GetProvider().IsAvailable())
-	{
-		FMessageLog("SourceControl").Error(LOCTEXT("SourceControlServerUnavailable", "Source control server is currently not available."));
-		return false;
-	}
-
-	// mark for add now if needed
-	ISourceControlProvider& Provider = ISourceControlModule::Get().GetProvider();
-	FSourceControlStatePtr SourceControlState = Provider.GetState(InFilePath, EStateCacheUsage::Use);
-	if (!SourceControlState.IsValid())
-	{
-		// Improper or invalid SCC state
-		FFormatNamedArguments Arguments;
-		Arguments.Add(TEXT("InFilePath"), FText::FromString(InFilePath));
-		FMessageLog("SourceControl").Error(FText::Format(LOCTEXT("CouldNotDetermineState", "Could not determine source control state of file '{InFilePath}'."), Arguments));
-		return false;
-	}
-
-	// add it if necessary
-	if (!SourceControlState->IsSourceControlled() || SourceControlState->IsUnknown())
-	{
-		ECommandResult::Type Result = Provider.Execute(ISourceControlOperation::Create<FMarkForAdd>(), InFilePath);
-		if (Result != ECommandResult::Succeeded)
-		{
-			FFormatNamedArguments Arguments;
-			Arguments.Add(TEXT("InFilePath"), FText::FromString(InFilePath));
-			FMessageLog("SourceControl").Error(FText::Format(LOCTEXT("MarkForAddFailed", "Failed to add file '{InFilePath}'."), Arguments));
-			return false;
-		}
-	}
-
-	return true;
-}
-
-bool USourceControlHelpers::CheckOutFile( const FString& InFilePath )
-{
-	if ( InFilePath.IsEmpty() )
-	{
-		FMessageLog("SourceControl").Error(LOCTEXT("UnspecifiedCheckoutFile", "Check out file not specified"));
-		return false;
-	}
-
-	if( !ISourceControlModule::Get().IsEnabled() )
-	{
-		FMessageLog("SourceControl").Error(LOCTEXT("SourceControlDisabled", "Source control is not enabled."));
-		return false;
-	}
-
-	if( !ISourceControlModule::Get().GetProvider().IsAvailable() )
-	{
-		FMessageLog("SourceControl").Error(LOCTEXT("SourceControlServerUnavailable", "Source control server is currently not available."));
-		return false;
-	}
-
-	bool bSuccessfullyCheckedOut = false;
-	TArray<FString> FilesToBeCheckedOut;
-	FilesToBeCheckedOut.Add( InFilePath );
-
-	ISourceControlProvider& SourceControlProvider = ISourceControlModule::Get().GetProvider();
-	FSourceControlStatePtr SourceControlState = SourceControlProvider.GetState( InFilePath, EStateCacheUsage::ForceUpdate );
-	if(SourceControlState.IsValid())
-	{
-		FString SimultaneousCheckoutUser;
-		if( SourceControlState->IsAdded() ||
-			SourceControlState->IsCheckedOut())
-		{
-			// Already checked out or opened for add
-			bSuccessfullyCheckedOut = true;
-		}
-		else
-		{
-			if(SourceControlState->CanCheckout())
-			{
-				bSuccessfullyCheckedOut = (SourceControlProvider.Execute( ISourceControlOperation::Create<FCheckOut>(), FilesToBeCheckedOut ) == ECommandResult::Succeeded);
-				if (!bSuccessfullyCheckedOut)
-				{
-					FFormatNamedArguments Arguments;
-					Arguments.Add( TEXT("InFilePath"), FText::FromString(InFilePath) );
-					FMessageLog("SourceControl").Error(FText::Format(LOCTEXT("CheckoutFailed", "Failed to check out file '{InFilePath}'."), Arguments));
-				}
-			}
-			else if(!SourceControlState->IsSourceControlled())
-			{
-				bSuccessfullyCheckedOut = (SourceControlProvider.Execute( ISourceControlOperation::Create<FMarkForAdd>(), FilesToBeCheckedOut ) == ECommandResult::Succeeded);
-				if (!bSuccessfullyCheckedOut)
-				{
-					FFormatNamedArguments Arguments;
-					Arguments.Add( TEXT("InFilePath"), FText::FromString(InFilePath) );
-					FMessageLog("SourceControl").Error(FText::Format(LOCTEXT("AddFailed", "Failed to add file '{InFilePath}' to source control."), Arguments));
-				}
-			}
-			else if(!SourceControlState->IsCurrent())
-			{
-				FFormatNamedArguments Arguments;
-				Arguments.Add( TEXT("InFilePath"), FText::FromString(InFilePath) );
-				FMessageLog("SourceControl").Error(FText::Format(LOCTEXT("NotAtHeadRevision", "File '{InFilePath}' is not at head revision."), Arguments));
-			}
-			else if(SourceControlState->IsCheckedOutOther(&(SimultaneousCheckoutUser)))
-			{
-				FFormatNamedArguments Arguments;
-				Arguments.Add( TEXT("InFilePath"), FText::FromString(InFilePath) );
-				Arguments.Add( TEXT("SimultaneousCheckoutUser"), FText::FromString(SimultaneousCheckoutUser) );
-				FMessageLog("SourceControl").Error(FText::Format(LOCTEXT("SimultaneousCheckout", "File '{InFilePath}' is checked out by another ('{SimultaneousCheckoutUser}')."), Arguments));
-			}
-			else
-			{
-				// Improper or invalid SCC state
-				FFormatNamedArguments Arguments;
-				Arguments.Add( TEXT("InFilePath"), FText::FromString(InFilePath) );
-				FMessageLog("SourceControl").Error(FText::Format(LOCTEXT("CouldNotDetermineState", "Could not determine source control state of file '{InFilePath}'."), Arguments));
-			}
-		}
-	}
-	else
-	{
-		// Improper or invalid SCC state
-		FFormatNamedArguments Arguments;
-		Arguments.Add( TEXT("InFilePath"), FText::FromString(InFilePath) );
-		FMessageLog("SourceControl").Error(FText::Format(LOCTEXT("CouldNotDetermineState", "Could not determine source control state of file '{InFilePath}'."), Arguments));
-	}
-
-	return bSuccessfullyCheckedOut;
-}
 
 bool USourceControlHelpers::CheckoutOrMarkForAdd( const FString& InDestFile, const FText& InFileDescription, const FOnPostCheckOut& OnPostCheckOut, FText& OutFailReason )
 {
@@ -392,6 +1092,7 @@ bool USourceControlHelpers::CheckoutOrMarkForAdd( const FString& InDestFile, con
 	return bSucceeded;
 }
 
+
 bool USourceControlHelpers::CopyFileUnderSourceControl( const FString& InDestFile, const FString& InSourceFile, const FText& InFileDescription, FText& OutFailReason)
 {
 	struct Local
@@ -413,6 +1114,7 @@ bool USourceControlHelpers::CopyFileUnderSourceControl( const FString& InDestFil
 	return CheckoutOrMarkForAdd(InDestFile, InFileDescription, FOnPostCheckOut::CreateStatic(&Local::CopyFile, InSourceFile), OutFailReason);
 }
 
+
 bool USourceControlHelpers::BranchPackage( UPackage* DestPackage, UPackage* SourcePackage )
 {
 	if(ISourceControlModule::Get().IsEnabled())
@@ -435,19 +1137,53 @@ bool USourceControlHelpers::BranchPackage( UPackage* DestPackage, UPackage* Sour
 }
 
 
+const FString& USourceControlHelpers::GetSettingsIni()
+{
+	if (ISourceControlModule::Get().GetUseGlobalSettings())
+	{
+		return GetGlobalSettingsIni();
+	}
+	else
+	{
+		static FString SourceControlSettingsIni;
+		if (SourceControlSettingsIni.Len() == 0)
+		{
+			const FString SourceControlSettingsDir = FPaths::GeneratedConfigDir();
+			FConfigCacheIni::LoadGlobalIniFile(SourceControlSettingsIni, TEXT("SourceControlSettings"), nullptr, false, false, true, *SourceControlSettingsDir);
+		}
+		return SourceControlSettingsIni;
+	}
+}
+
+
+const FString& USourceControlHelpers::GetGlobalSettingsIni()
+{
+	static FString SourceControlGlobalSettingsIni;
+	if (SourceControlGlobalSettingsIni.Len() == 0)
+	{
+		const FString SourceControlSettingsDir = FPaths::EngineSavedDir() + TEXT("Config/");
+		FConfigCacheIni::LoadGlobalIniFile(SourceControlGlobalSettingsIni, TEXT("SourceControlSettings"), nullptr, false, false, true, *SourceControlSettingsDir);
+	}
+	return SourceControlGlobalSettingsIni;
+}
+
+
 FScopedSourceControl::FScopedSourceControl()
 {
 	ISourceControlModule::Get().GetProvider().Init();
 }
 
+
 FScopedSourceControl::~FScopedSourceControl()
 {
 	ISourceControlModule::Get().GetProvider().Close();
 }
 
+
 ISourceControlProvider& FScopedSourceControl::GetProvider()
 {
 	return ISourceControlModule::Get().GetProvider();
 }
 
+
 #undef LOCTEXT_NAMESPACE
diff --git a/Engine/Source/Developer/SourceControl/Private/SourceControlModule.cpp b/Engine/Source/Developer/SourceControl/Private/SourceControlModule.cpp
--- a/Engine/Source/Developer/SourceControl/Private/SourceControlModule.cpp
+++ b/Engine/Source/Developer/SourceControl/Private/SourceControlModule.cpp
@@ -3,6 +3,8 @@
 #include "SourceControlModule.h"
 #include "Features/IModularFeatures.h"
 #include "SSourceControlLogin.h"
+#include "SourceControlOperations.h"
+#include "SourceControlHelpers.h"
 
 #if SOURCE_CONTROL_WITH_SLATE
 	#include "Widgets/DeclarativeSyntaxSupport.h"
diff --git a/Engine/Source/Developer/SourceControl/Public/ISourceControlModule.h b/Engine/Source/Developer/SourceControl/Public/ISourceControlModule.h
--- a/Engine/Source/Developer/SourceControl/Public/ISourceControlModule.h
+++ b/Engine/Source/Developer/SourceControl/Public/ISourceControlModule.h
@@ -85,7 +85,7 @@ public:
 	virtual void QueueStatusUpdate(const FString& InFilename) = 0;
 
 	/**
-	 * Check whether source control is enabled.
+	 * Check whether source control is enabled.	Specifically, it returns true if a source control provider is set (regardless of whether the provider is available) and false if no provider is set.
 	 */
 	virtual bool IsEnabled() const = 0;
 
