Copying //UE4/Dev-Enterprise to //UE4/Dev-Main (Source: //UE4/Dev-Enterprise @ 3972172)

diff --git a/Engine/Source/Developer/SourceControl/Public/ISourceControlOperation.h b/Engine/Source/Developer/SourceControl/Public/ISourceControlOperation.h
--- a/Engine/Source/Developer/SourceControl/Public/ISourceControlOperation.h
+++ b/Engine/Source/Developer/SourceControl/Public/ISourceControlOperation.h
@@ -4,6 +4,25 @@
 
 #include "CoreMinimal.h"
 
+
+/** Accumulated error and info messages for a source control operation.  */
+struct FSourceControlResultInfo
+{
+	/** Append any messages from another FSourceControlResultInfo, ensuring to keep any already accumulated info. */
+	void Append(const FSourceControlResultInfo& InResultInfo)
+	{
+		InfoMessages.Append(InResultInfo.InfoMessages);
+		ErrorMessages.Append(InResultInfo.ErrorMessages);
+	}
+
+	/** Info and/or warning message storage */
+	TArray<FText> InfoMessages;
+
+	/** Potential error message storage */
+	TArray<FText> ErrorMessages;
+};
+
+
 class ISourceControlOperation : public TSharedFromThis<ISourceControlOperation, ESPMode::ThreadSafe>
 {
 public:
@@ -16,9 +35,38 @@ public:
 	virtual FName GetName() const = 0;
 
 	/** Get the string to display when this operation is in progress */
-	virtual FText GetInProgressString() const 
-	{ 
-		return FText(); 
+	virtual FText GetInProgressString() const
+	{
+		return FText();
+	}
+
+	/** Retrieve any info or error messages that may have accumulated during the operation. */
+	virtual const FSourceControlResultInfo& GetResultInfo() const
+	{
+		// Implemented in subclasses
+		static const FSourceControlResultInfo ResultInfo = FSourceControlResultInfo();
+		return ResultInfo;
+	}
+
+	/** Add info/warning message. */
+	virtual void AddInfoMessge(const FText& InInfo)
+	{
+		// Implemented in subclasses
+	}
+
+	/** Add error message. */
+	virtual void AddErrorMessge(const FText& InError)
+	{
+		// Implemented in subclasses
+	}
+
+	/**
+	 * Append any info or error messages that may have accumulated during the operation prior
+	 * to returning a result, ensuring to keep any already accumulated info.
+	 */
+	virtual void AppendResultInfo(const FSourceControlResultInfo& InResultInfo)
+	{
+		// Implemented in subclasses
 	}
 
 	/** Factory method for easier operation creation */
@@ -27,6 +75,7 @@ public:
 	{
 		return MakeShareable( new Type() );
 	}
+
 };
 
 typedef TSharedRef<class ISourceControlOperation, ESPMode::ThreadSafe> FSourceControlOperationRef;
diff --git a/Engine/Source/Developer/SourceControl/Public/ISourceControlProvider.h b/Engine/Source/Developer/SourceControl/Public/ISourceControlProvider.h
--- a/Engine/Source/Developer/SourceControl/Public/ISourceControlProvider.h
+++ b/Engine/Source/Developer/SourceControl/Public/ISourceControlProvider.h
@@ -4,10 +4,8 @@
 
 #include "CoreMinimal.h"
 #include "ISourceControlOperation.h"
-#include "SourceControlOperations.h"
-#include "Features/IModularFeature.h"
 #include "ISourceControlState.h"
-#include "SourceControlHelpers.h"
+#include "Features/IModularFeature.h"
 
 #ifndef SOURCE_CONTROL_WITH_SLATE
 	#error "SOURCE_CONTROL_WITH_SLATE must be defined. Did you forget a dependency on the 'SourceControl' module?"
@@ -81,13 +79,13 @@ public:
 	 */
 	virtual ~ISourceControlProvider() {}
 
-	/** 
+	/**
 	 * Initialize source control provider.
 	 * @param	bForceConnection	If set, this flag forces the provider to attempt a connection to its server.
 	 */
 	virtual void Init(bool bForceConnection = true) = 0;
 
-	/** 
+	/**
 	 * Shut down source control provider.
 	 */
 	virtual void Close() = 0;
@@ -98,7 +96,7 @@ public:
 	/** Get the source control status as plain, human-readable text */
 	virtual FText GetStatusText() const = 0;
 
-	/** Quick check if source control is enabled */
+	/** Quick check if source control is enabled. Specifically, it returns true if a source control provider is set (regardless of whether the provider is available) and false if no provider is set. So all providers except the stub DefalutSourceProvider will return true. */
 	virtual bool IsEnabled() const = 0;
 
 	/**
@@ -114,15 +112,10 @@ public:
 	 * This is just a wrapper around Execute().
 	 * @param	InPassword						The password (if required) to use when logging in.
 	 * @param	InConcurrency					How to execute the operation, blocking or asynchronously on another thread.
-	 * @param	InOperationCompleteDelegate		Delegate to call when the operation is completed. This is called back internal to this call when executed on the main thread, or from Tick() when queued for asynchronous execution.
+	 * @param	InOperationCompleteDelegate		Delegate to call when the operation is completed. This is called back internal to this call when executed on the main thread, or from Tick() when queued for asynchronous execution. If the provider is not enabled or if the command is not suported the delegate is immediately called with ECommandResult::Failed.
 	 * @return the result of the operation.
 	 */
-	virtual ECommandResult::Type Login( const FString& InPassword = FString(), EConcurrency::Type InConcurrency = EConcurrency::Synchronous, const FSourceControlOperationComplete& InOperationCompleteDelegate = FSourceControlOperationComplete() )
-	{
-		TSharedRef<FConnect, ESPMode::ThreadSafe> ConnectOperation = ISourceControlOperation::Create<FConnect>();
-		ConnectOperation->SetPassword(InPassword);
-		return Execute(ConnectOperation, InConcurrency, InOperationCompleteDelegate);
-	}
+	SOURCECONTROL_API virtual ECommandResult::Type Login( const FString& InPassword = FString(), EConcurrency::Type InConcurrency = EConcurrency::Synchronous, const FSourceControlOperationComplete& InOperationCompleteDelegate = FSourceControlOperationComplete() );
 
 	/**
 	* Register branches to query for state in addition to the current branch
@@ -130,7 +123,7 @@ public:
 	* @param	ContentRoot			Path to the content root for branch mapping
 	*/
 	virtual void RegisterStateBranches(const TArray<FString>& BranchNames, const FString& ContentRoot) = 0;
-	
+
 	/**
 	*Gets the state index of the specified branch, higher index branches are generally closer to releases
 	* @param	BranchName			Names of the branches to query
@@ -150,35 +143,17 @@ public:
 	/**
 	 * Helper overload for state retrieval, see GetState().
 	 */
-	virtual ECommandResult::Type GetState( const TArray<UPackage*>& InPackages, TArray< TSharedRef<ISourceControlState, ESPMode::ThreadSafe> >& OutState, EStateCacheUsage::Type InStateCacheUsage )
-	{
-		TArray<FString> Files = SourceControlHelpers::PackageFilenames(InPackages);
-		return GetState(Files, OutState, InStateCacheUsage);
-	}
+	SOURCECONTROL_API virtual ECommandResult::Type GetState( const TArray<UPackage*>& InPackages, TArray< TSharedRef<ISourceControlState, ESPMode::ThreadSafe> >& OutState, EStateCacheUsage::Type InStateCacheUsage );
 
 	/**
 	 * Helper overload for state retrieval, see GetState().
 	 */
-	virtual TSharedPtr<ISourceControlState, ESPMode::ThreadSafe> GetState( const FString& InFile, EStateCacheUsage::Type InStateCacheUsage )
-	{
-		TArray<FString> Files;
-		TArray< TSharedRef<ISourceControlState, ESPMode::ThreadSafe> > States;
-		Files.Add(InFile);
-		if(GetState(Files, States, InStateCacheUsage) == ECommandResult::Succeeded)
-		{
-			TSharedRef<ISourceControlState, ESPMode::ThreadSafe> State = States[0];
-			return State;
-		}
-		return NULL;
-	}
+	SOURCECONTROL_API virtual TSharedPtr<ISourceControlState, ESPMode::ThreadSafe> GetState( const FString& InFile, EStateCacheUsage::Type InStateCacheUsage );
 
 	/**
 	 * Helper overload for state retrieval, see GetState().
 	 */
-	virtual TSharedPtr<ISourceControlState, ESPMode::ThreadSafe> GetState( const UPackage* InPackage, EStateCacheUsage::Type InStateCacheUsage )
-	{
-		return GetState(SourceControlHelpers::PackageFilename(InPackage), InStateCacheUsage);
-	}
+	SOURCECONTROL_API virtual TSharedPtr<ISourceControlState, ESPMode::ThreadSafe> GetState( const UPackage* InPackage, EStateCacheUsage::Type InStateCacheUsage );
 
 	/**
 	 * Get all cached source control state objects for which the supplied predicate returns true
@@ -200,7 +175,7 @@ public:
 	 * @param	InOperation						The operation to perform.
 	 * @param	InFiles							The files to operate on.
 	 * @param	InConcurrency					How to execute the operation, blocking or asynchronously on another thread.
-	 * @param	InOperationCompleteDelegate		Delegate to call when the operation is completed. This is called back internal to this call when executed on the main thread, or from Tick() when queued for asynchronous execution.
+	 * @param	InOperationCompleteDelegate		Delegate to call when the operation is completed. This is called back internal to this call when executed on the main thread, or from Tick() when queued for asynchronous execution. If the provider is not enabled or if the command is not suported the delegate is immediately called with ECommandResult::Failed.
 	 * @return the result of the operation.
 	 */
 	virtual ECommandResult::Type Execute( const TSharedRef<ISourceControlOperation, ESPMode::ThreadSafe>& InOperation, const TArray<FString>& InFiles, EConcurrency::Type InConcurrency = EConcurrency::Synchronous, const FSourceControlOperationComplete& InOperationCompleteDelegate = FSourceControlOperationComplete() ) = 0;
@@ -208,37 +183,22 @@ public:
 	/**
 	 * Helper overload for operation execution, see Execute().
 	 */
-	virtual ECommandResult::Type Execute( const TSharedRef<ISourceControlOperation, ESPMode::ThreadSafe>& InOperation, const EConcurrency::Type InConcurrency = EConcurrency::Synchronous, const FSourceControlOperationComplete& InOperationCompleteDelegate = FSourceControlOperationComplete() )
-	{
-		return Execute(InOperation, TArray<FString>(), InConcurrency, InOperationCompleteDelegate);
-	}
+	SOURCECONTROL_API virtual ECommandResult::Type Execute( const TSharedRef<ISourceControlOperation, ESPMode::ThreadSafe>& InOperation, const EConcurrency::Type InConcurrency = EConcurrency::Synchronous, const FSourceControlOperationComplete& InOperationCompleteDelegate = FSourceControlOperationComplete() );
 
 	/**
 	 * Helper overload for operation execution, see Execute().
 	 */
-	virtual ECommandResult::Type Execute( const TSharedRef<ISourceControlOperation, ESPMode::ThreadSafe>& InOperation, const UPackage* InPackage, const EConcurrency::Type InConcurrency = EConcurrency::Synchronous, const FSourceControlOperationComplete& InOperationCompleteDelegate = FSourceControlOperationComplete() )
-	{
-		return Execute(InOperation, SourceControlHelpers::PackageFilename(InPackage), InConcurrency, InOperationCompleteDelegate);
-	}
+	SOURCECONTROL_API virtual ECommandResult::Type Execute( const TSharedRef<ISourceControlOperation, ESPMode::ThreadSafe>& InOperation, const UPackage* InPackage, const EConcurrency::Type InConcurrency = EConcurrency::Synchronous, const FSourceControlOperationComplete& InOperationCompleteDelegate = FSourceControlOperationComplete() );
 
 	/**
 	 * Helper overload for operation execution, see Execute().
 	 */
-	virtual ECommandResult::Type Execute( const TSharedRef<ISourceControlOperation, ESPMode::ThreadSafe>& InOperation, const FString& InFile, const EConcurrency::Type InConcurrency = EConcurrency::Synchronous, const FSourceControlOperationComplete& InOperationCompleteDelegate = FSourceControlOperationComplete() )
-	{
-		TArray<FString> FileArray;
-		FileArray.Add(InFile);
-		return Execute(InOperation, FileArray, InConcurrency, InOperationCompleteDelegate);
-	}
+	SOURCECONTROL_API virtual ECommandResult::Type Execute( const TSharedRef<ISourceControlOperation, ESPMode::ThreadSafe>& InOperation, const FString& InFile, const EConcurrency::Type InConcurrency = EConcurrency::Synchronous, const FSourceControlOperationComplete& InOperationCompleteDelegate = FSourceControlOperationComplete() );
 
 	/**
 	 * Helper overload for operation execution, see Execute().
 	 */
-	virtual ECommandResult::Type Execute( const TSharedRef<ISourceControlOperation, ESPMode::ThreadSafe>& InOperation, const TArray<UPackage*>& InPackages, const EConcurrency::Type InConcurrency = EConcurrency::Synchronous, const FSourceControlOperationComplete& InOperationCompleteDelegate = FSourceControlOperationComplete() )
-	{
-		TArray<FString> FileArray = SourceControlHelpers::PackageFilenames(InPackages);
-		return Execute(InOperation, FileArray, InConcurrency, InOperationCompleteDelegate);
-	}
+	SOURCECONTROL_API virtual ECommandResult::Type Execute( const TSharedRef<ISourceControlOperation, ESPMode::ThreadSafe>& InOperation, const TArray<UPackage*>& InPackages, const EConcurrency::Type InConcurrency = EConcurrency::Synchronous, const FSourceControlOperationComplete& InOperationCompleteDelegate = FSourceControlOperationComplete() );
 
 	/**
 	 * Check to see if we can cancel an operation.
@@ -258,16 +218,7 @@ public:
 	 * @param	InLabelName	String specifying the label name
 	 * @return the label, if any
 	 */
-	virtual TSharedPtr<class ISourceControlLabel> GetLabel( const FString& InLabelName ) const
-	{
-		TArray< TSharedRef<class ISourceControlLabel> > Labels = GetLabels(InLabelName);
-		if(Labels.Num() > 0)
-		{
-			return Labels[0];
-		}
-
-		return NULL;
-	}
+	SOURCECONTROL_API virtual TSharedPtr<class ISourceControlLabel> GetLabel( const FString& InLabelName ) const;
 
 	/**
 	 * Get an array of labels matching the passed-in spec.
@@ -297,7 +248,7 @@ public:
 	virtual void Tick() = 0;
 
 #if SOURCE_CONTROL_WITH_SLATE
-	/** 
+	/**
 	 * Create a settings widget for display in the login window.
 	 * @returns a widget used to edit the providers settings required prior to connection.
 	 */
diff --git a/Engine/Source/Developer/SourceControl/Public/SourceControlHelpers.h b/Engine/Source/Developer/SourceControl/Public/SourceControlHelpers.h
--- a/Engine/Source/Developer/SourceControl/Public/SourceControlHelpers.h
+++ b/Engine/Source/Developer/SourceControl/Public/SourceControlHelpers.h
@@ -12,6 +12,104 @@
 
 class ISourceControlProvider;
 
+/**
+ * Snapshot of source control state of for a file
+ * @see	USourceControlHelpers::QueryFileState()
+ */
+USTRUCT(BlueprintType)
+struct FSourceControlState
+{
+	GENERATED_BODY()
+
+public:
+
+	/** Get the local filename that this state represents */
+	UPROPERTY(VisibleAnywhere, BlueprintReadOnly, Category = "Editor Scripting | Editor Source Control Helpers")
+	FString Filename;
+
+	/** Indicates whether this source control state has valid information (true) or not (false) */
+	UPROPERTY(VisibleAnywhere, BlueprintReadOnly, Category = "Editor Scripting | Editor Source Control Helpers")
+	bool bIsValid;
+
+	/** Determine if we know anything about the source control state of this file */
+	UPROPERTY(VisibleAnywhere, BlueprintReadOnly, Category = "Editor Scripting | Editor Source Control Helpers")
+	bool bIsUnknown;
+
+	/** Determine if this file can be checked in. */
+	UPROPERTY(VisibleAnywhere, BlueprintReadOnly, Category = "Editor Scripting | Editor Source Control Helpers")
+	bool bCanCheckIn;
+
+	/** Determine if this file can be checked out */
+	UPROPERTY(VisibleAnywhere, BlueprintReadOnly, Category = "Editor Scripting | Editor Source Control Helpers")
+	bool bCanCheckOut;
+
+	/** Determine if this file is checked out */
+	UPROPERTY(VisibleAnywhere, BlueprintReadOnly, Category = "Editor Scripting | Editor Source Control Helpers")
+	bool bIsCheckedOut;
+
+	/** Determine if this file is up-to-date with the version in source control */
+	UPROPERTY(VisibleAnywhere, BlueprintReadOnly, Category = "Editor Scripting | Editor Source Control Helpers")
+	bool bIsCurrent;
+
+	/** Determine if this file is under source control */
+	UPROPERTY(VisibleAnywhere, BlueprintReadOnly, Category = "Editor Scripting | Editor Source Control Helpers")
+	bool bIsSourceControlled;
+
+	/**
+	 * Determine if this file is marked for add
+	 * @note	if already checked in then not considered mid add
+	 */
+	UPROPERTY(VisibleAnywhere, BlueprintReadOnly, Category = "Editor Scripting | Editor Source Control Helpers")
+	bool bIsAdded;
+
+	/** Determine if this file is marked for delete */
+	UPROPERTY(VisibleAnywhere, BlueprintReadOnly, Category = "Editor Scripting | Editor Source Control Helpers")
+	bool bIsDeleted;
+
+	/** Determine if this file is ignored by source control */
+	UPROPERTY(VisibleAnywhere, BlueprintReadOnly, Category = "Editor Scripting | Editor Source Control Helpers")
+	bool bIsIgnored;
+
+	/** Determine if source control allows this file to be edited */
+	UPROPERTY(VisibleAnywhere, BlueprintReadOnly, Category = "Editor Scripting | Editor Source Control Helpers")
+	bool bCanEdit;
+
+	/** Determine if source control allows this file to be deleted. */
+	UPROPERTY(VisibleAnywhere, BlueprintReadOnly, Category = "Editor Scripting | Editor Source Control Helpers")
+	bool bCanDelete;
+
+	/** Determine if this file is modified compared to the version in source control. */
+	UPROPERTY(VisibleAnywhere, BlueprintReadOnly, Category = "Editor Scripting | Editor Source Control Helpers")
+	bool bIsModified;
+
+	/** 
+	 * Determine if this file can be added to source control (i.e. is part of the directory 
+	 * structure currently under source control) 
+	 */
+	UPROPERTY(VisibleAnywhere, BlueprintReadOnly, Category = "Editor Scripting | Editor Source Control Helpers")
+	bool bCanAdd;
+
+	/** Determine if this file is in a conflicted state */
+	UPROPERTY(VisibleAnywhere, BlueprintReadOnly, Category = "Editor Scripting | Editor Source Control Helpers")
+	bool bIsConflicted;
+
+	/** Determine if this file can be reverted, i.e. discard changes and the file will no longer be checked-out. */
+	UPROPERTY(VisibleAnywhere, BlueprintReadOnly, Category = "Editor Scripting | Editor Source Control Helpers")
+	bool bCanRevert;
+
+	/** Determine if this file is checked out by someone else */
+	UPROPERTY(VisibleAnywhere, BlueprintReadOnly, Category = "Editor Scripting | Editor Source Control Helpers")
+	bool bIsCheckedOutOther;
+
+	/**
+	 * Get name of other user who this file already checked out or "" if no other user has it checked out
+	 */
+	UPROPERTY(VisibleAnywhere, BlueprintReadOnly, Category = "Editor Scripting | Editor Source Control Helpers")
+	FString CheckedOutOther;
+
+};  // FSourceControlState
+
+
 /** 
  * Delegate used for performing operation on files that may need a checkout, but before they are added to source control 
  * @param	InDestFile			The filename that was potentially checked out
@@ -24,69 +122,292 @@
 // For backwards compatibility
 typedef class USourceControlHelpers SourceControlHelpers;
 
-UCLASS(transient)
-class USourceControlHelpers : public UObject
-{
+/**
+ * Editor source control common functionality.
+ *
+ * @note Many of these source control methods use *smart* file strings which can be one of:
+ *   - fully qualified path
+ *   - relative path
+ *   - long package name
+ *   - asset path
+ *   - export text path (often stored on clipboard)
+ *
+ *   For example:
+ *	 - D:\Epic\Dev-Ent\Projects\Python3rdBP\Content\Mannequin\Animations\ThirdPersonIdle.uasset
+ *	 - Content\Mannequin\Animations\ThirdPersonIdle.uasset
+ *	 - /Game/Mannequin/Animations/ThirdPersonIdle
+ *	 - /Game/Mannequin/Animations/ThirdPersonIdle.ThirdPersonIdle
+ *	 - AnimSequence'/Game/Mannequin/Animations/ThirdPersonIdle.ThirdPersonIdle'
+ */
+UCLASS(Abstract, Transient, meta = (ScriptName = "SourceControl"))
+class SOURCECONTROL_API USourceControlHelpers : public UObject
+{	// Note - would use UBlueprintFunctionLibrary however it requires Engine module which breaks some
+	// other modules that depend on the SourceControl module and expect to be Engine independent.
+
 	GENERATED_BODY()
+
 public:
 	/**
-	 * Helper function to get the ini filename for storing source control settings
-	 * @return the filename
+	 * Determine the name of the current source control provider.
+	 * @return	the name of the current source control provider. If one is not set then "None" is returned.
 	 */
-	static SOURCECONTROL_API const FString& GetSettingsIni();
+	UFUNCTION(BlueprintPure, Category = "Editor Scripting | Editor Source Control Helpers")
+	static FString CurrentProvider();
 
 	/**
-	 * Helper function to get the ini filename for storing global source control settings
-	 * @return the filename
+	 * Determine if there is a source control system enabled
+	 * @return	true if enabled, false if not
 	 */
-	static SOURCECONTROL_API const FString& GetGlobalSettingsIni();
+	UFUNCTION(BlueprintPure, Category = "Editor Scripting | Editor Source Control Helpers")
+	static bool IsEnabled();
 
 	/**
+	 * Quick check if currently set source control provider is enabled and available for use
+	 * (server-based providers can use this to return whether the server is available or not)
+	 *
+	 * @return	true if source control is available, false if it is not
+	 */
+	UFUNCTION(BlueprintPure, Category = "Editor Scripting | Editor Source Control Helpers")
+	static bool IsAvailable();
+
+	/**
+	* Get status text set by SourceControl system if an error occurs regardless whether bSilent is set or not.
+	* Only set if there was an error.
+	*/
+	UFUNCTION(BlueprintPure, Category = "Editor Scripting | Editor Source Control Helpers")
+	static FText LastErrorMsg();
+
+	/**
+	 * Use currently set source control provider to check out a file.
+	 * @note	Blocks until action is complete.
+	 *
+	 * @param	InFile		The file to check out - can be either fully qualified path, relative path, long package name, asset path or export text path (often stored on clipboard)
+	 * @param	bSilent		if false (default) then write out any error info to the Log. Any error text can be retrieved by LastErrorMsg() regardless.
+	 * @return	true if succeeded, false if failed and can call LastErrorMsg() for more info.
+	 */
+	UFUNCTION(BlueprintCallable, Category = "Editor Scripting | Editor Source Control Helpers")
+	static bool CheckOutFile(const FString& InFile, bool bSilent = false);
+
+	/**
+	 * Use currently set source control provider to check out specified files.
+	 * @note	Blocks until action is complete.
+	 *
+	 * @param	InFiles		Files to check out - can be either fully qualified path, relative path, long package name, asset path or export text path (often stored on clipboard)
+	 * @param	bSilent		if false (default) then write out any error info to the Log. Any error text can be retrieved by LastErrorMsg() regardless.
+	 * @return	true if succeeded, false if failed and can call LastErrorMsg() for more info.
+	 */
+	UFUNCTION(BlueprintCallable, Category = "Editor Scripting | Editor Source Control Helpers")
+	static bool CheckOutFiles(const TArray<FString>& InFiles, bool bSilent = false);
+
+	/**
+	 * Use currently set source control provider to check out file or mark it for add.
+	 * @note	Blocks until action is complete.
+	 *
+	 * @param	InFile		The file to check out/add - can be either fully qualified path, relative path, long package name, asset path or export text path (often stored on clipboard)
+	 * @param	bSilent		if false (default) then write out any error info to the Log. Any error text can be retrieved by LastErrorMsg() regardless.
+	 * @return	true if succeeded, false if failed and can call LastErrorMsg() for more info.
+	 */
+	UFUNCTION(BlueprintCallable, Category = "Editor Scripting | Editor Source Control Helpers")
+	static bool CheckOutOrAddFile(const FString& InFile, bool bSilent = false);
+
+	/**
+	 * Helper function perform an operation on files in our 'source controlled' directories, handling checkout/add etc.
+	 * @note	Blocks until action is complete. Older C++ only version of CheckOutOrAddFile().
+	 *
+	 * @param	InDestFile			The path to the destination file
+	 * @param	InFileDescription	Description of the file to display to the user, e.g. "Text" or "Image"
+	 * @param	OnPostCheckOut		Delegate used for performing operation on files that may need a checkout, but before they are added to source control 
+	 * @param	OutFailReason		Text describing why the operation failed
+	 * @return	Success or failure of the operation
+	 * @see     CheckOutOrAddFile()
+	 */
+	static bool CheckoutOrMarkForAdd(const FString& InDestFile, const FText& InFileDescription, const FOnPostCheckOut& OnPostCheckOut, FText& OutFailReason);
+
+	/**
+	 * Use currently set source control provider to mark a file for add. Does nothing (and returns true) if the file is already under SC
+	 * @note	Blocks until action is complete.
+	 *
+	 * @param	InFile		The file to add - can be either fully qualified path, relative path, long package name, asset path or export text path (often stored on clipboard)
+	 * @param	bSilent		if false (default) then write out any error info to the Log. Any error text can be retrieved by LastErrorMsg() regardless.
+	 * @return	true if succeeded, false if failed and can call LastErrorMsg() for more info.
+	 */
+	UFUNCTION(BlueprintCallable, Category = "Editor Scripting | Editor Source Control Helpers")
+	static bool MarkFileForAdd(const FString& InFile, bool bSilent = false);
+
+	/**
+	 * Use currently set source control provider to mark files for add. Does nothing (and returns true) for any file that is already under SC
+	 * @note	Blocks until action is complete.
+	 *
+	 * @param	InFiles		Files to check out - can be either fully qualified path, relative path, long package name, asset path or export text path (often stored on clipboard)
+	 * @param	bSilent		if false (default) then write out any error info to the Log. Any error text can be retrieved by LastErrorMsg() regardless.
+	 * @return	true if succeeded, false if failed and can call LastErrorMsg() for more info.
+	 */
+	UFUNCTION(BlueprintCallable, Category = "Editor Scripting | Editor Source Control Helpers")
+	static bool MarkFilesForAdd(const TArray<FString>& InFiles, bool bSilent = false);
+
+	/**
+	 * Use currently set source control provider to remove file from source control and
+	 * delete the file.
+	 * @note	Blocks until action is complete.
+	 *
+	 * @param	InFile		The file to delete - can be either fully qualified path, relative path, long package name, asset path or export text path (often stored on clipboard)
+	 * @param	bSilent		if false (default) then write out any error info to the Log. Any error text can be retrieved by LastErrorMsg() regardless.
+	 * @return	true if succeeded, false if failed and can call LastErrorMsg() for more info.
+	 */
+	UFUNCTION(BlueprintCallable, Category = "Editor Scripting | Editor Source Control Helpers")
+	static bool MarkFileForDelete(const FString& InFile, bool bSilent = false);
+
+	/**
+	 * Use currently set source control provider to revert a file regardless whether any changes will be lost or not.
+	 * @note	Blocks until action is complete.
+	 *
+	 * @param	InFile	The file to revert - can be either fully qualified path, relative path, long package name, asset path or export text path (often stored on clipboard)
+	 * @param	bSilent	if false (default) then write out any error info to the Log. Any error text can be retrieved by LastErrorMsg() regardless.
+	 * @return	true if succeeded, false if failed and can call LastErrorMsg() for more info.
+	 */
+	UFUNCTION(BlueprintCallable, Category = "Editor Scripting | Editor Source Control Helpers")
+	static bool RevertFile(const FString& InFile, bool bSilent = false);
+
+	/**
+	 * Use currently set source control provider to revert files regardless whether any changes will be lost or not.
+	 * @note	Blocks until action is complete.
+	 *
+	 * @param	InFiles	Files to revert - can be either fully qualified path, relative path, long package name, asset path or export text path (often stored on clipboard)
+	 * @param	bSilent	if false (default) then write out any error info to the Log. Any error text can be retrieved by LastErrorMsg() regardless.
+	 * @return	true if succeeded, false if failed and can call LastErrorMsg() for more info.
+	 */
+	UFUNCTION(BlueprintCallable, Category = "Editor Scripting | Editor Source Control Helpers")
+	static bool RevertFiles(const TArray<FString>& InFiles, bool bSilent = false);
+
+	/**
+	 * Use currently set source control provider to revert a file provided no changes have been made.
+	 * @note	Blocks until action is complete.
+	 *
+	 * @param	InFile	File to check out - can be either fully qualified path, relative path, long package name, asset path or export text path (often stored on clipboard)
+	 * @param	bSilent	if false (default) then write out any error info to the Log. Any error text can be retrieved by LastErrorMsg() regardless.
+	 * @return	true if succeeded, false if failed and can call LastErrorMsg() for more info.
+	 */
+	UFUNCTION(BlueprintCallable, Category = "Editor Scripting | Editor Source Control Helpers")
+	static bool RevertUnchangedFile(const FString& InFile, bool bSilent = false);
+
+	/**
+	 * Use currently set source control provider to revert files provided no changes have been made.
+	 * @note	Blocks until action is complete.
+	 *
+	 * @param	InFiles	Files to check out - can be either fully qualified path, relative path, long package name, asset path or export text path (often stored on clipboard)
+	 * @param	bSilent	if false (default) then write out any error info to the Log. Any error text can be retrieved by LastErrorMsg() regardless.
+	 * @return	true if succeeded, false if failed and can call LastErrorMsg() for more info.
+	 */
+	UFUNCTION(BlueprintCallable, Category = "Editor Scripting | Editor Source Control Helpers")
+	static bool RevertUnchangedFiles(const TArray<FString>& InFiles, bool bSilent = false);
+
+	/**
+	 * Helper function to get a list of files that are unchanged & revert them. This runs synchronous commands.
+	 * @note	Older C++ only version of RevertFiles(Files, false).
+	 * @see		RevertFiles()
+	 *
+	 * @param	InProvider	The provider to use
+	 * @param	InFiles		The files to operate on
+	 */
+	static void RevertUnchangedFiles(ISourceControlProvider& InProvider, const TArray<FString>& InFiles);
+
+	/**
+	 * Use currently set source control provider to check in a file.
+	 * @note	Blocks until action is complete.
+	 *
+	 * @param	InFile			The file to check in - can be either fully qualified path, relative path, long package name, asset path or export text path (often stored on clipboard)
+	 * @param	InDescription	Description for check in
+	 * @param	bSilent			if false (default) then write out any error info to the Log. Any error text can be retrieved by LastErrorMsg() regardless.
+	 * @return	true if succeeded, false if failed and can call LastErrorMsg() for more info.
+	 */
+	UFUNCTION(BlueprintCallable, Category = "Editor Scripting | Editor Source Control Helpers")
+	static bool CheckInFile(const FString& InFile, const FString& InDescription, bool bSilent = false);
+
+	/**
+	 * Use currently set source control provider to check in specified files.
+	 * @note	Blocks until action is complete.
+	 *
+	 * @param	InFiles			Files to check out - can be either fully qualified path, relative path, long package name, asset path or export text path (often stored on clipboard)
+	 * @param	InDescription	Description for check in
+	 * @param	bSilent			if false (default) then write out any error info to the Log. Any error text can be retrieved by LastErrorMsg() regardless.
+	 * @return	true if succeeded, false if failed and can call LastErrorMsg() for more info.
+	 */
+	UFUNCTION(BlueprintCallable, Category = "Editor Scripting | Editor Source Control Helpers")
+	static bool CheckInFiles(const TArray<FString>& InFiles, const FString& InDescription, bool bSilent = false);
+
+	/**
+	 * Use currently set source control provider to copy a file.
+	 * @note	Blocks until action is complete.
+	 *
+	 * @param	InSourcePath	Source file string to copy from - can be either fully qualified path, relative path, long package name, asset path or export text path (often stored on clipboard)
+	 * @param	InDestPath		Source file string to copy to - can be either fully qualified path, relative path, long package name, asset path or export text path (often stored on clipboard). If package, then uses same extension as source file.
+	 * @param	bSilent			if false (default) then write out any error info to the Log. Any error text can be retrieved by LastErrorMsg() regardless.
+	 * @return	true if succeeded, false if failed and can call LastErrorMsg() for more info.
+	 */
+	UFUNCTION(BlueprintCallable, Category = "Editor Scripting | Editor Source Control Helpers")
+	static bool CopyFile(const FString& InSourceFile, const FString& InDestFile, bool bSilent = false);
+
+	/**
+	 * Helper function to copy a file into our 'source controlled' directories, handling checkout/add etc.
+	 * @note	Blocks until action is complete. Older C++ only version of CopyFile().
+	 *
+	 * @param	InDestFile			The path to the destination file
+	 * @param	InSourceFile		The path to the source file
+	 * @param	InFileDescription	Description of the file to display to the user, e.g. "Text" or "Image"
+	 * @param	OutFailReason		Text describing why the operation failed
+	 * @return	Success or failure of the operation
+	 */
+	static bool CopyFileUnderSourceControl(const FString& InDestFile, const FString& InSourceFile, const FText& InFileDescription, FText& OutFailReason);
+
+	/**
+	 * Use currently set source control provider to query a file's source control state.
+	 * @note	Blocks until action is complete.
+	 *
+	 * @param	InFile			The file to query - can be either fully qualified path, relative path, long package name, asset path or export text path (often stored on clipboard)
+	 * @param	bSilent			if false (default) then write out any error info to the Log. Any error text can be retrieved by LastErrorMsg() regardless.
+	 * @return	Source control state - see USourceControlState. It will have bIsValid set to false if
+	 *			it could not have its values set.
+	 */
+	UFUNCTION(BlueprintCallable, Category = "Editor Scripting | Editor Source Control Helpers")
+	static FSourceControlState QueryFileState(const FString& InFile, bool bSilent = false);
+
+
+	/**
 	 * Helper function to get a filename for a package name.
 	 * @param	InPackageName	The package name to get the filename for
 	 * @return the filename
 	 */
-	UFUNCTION(BlueprintCallable, Category = "Editor Scripting | Editor Source Control Helpers")
-	static SOURCECONTROL_API FString PackageFilename( const FString& InPackageName );
+	static FString PackageFilename(const FString& InPackageName);
 
 	/**
 	 * Helper function to get a filename for a package.
 	 * @param	InPackage	The package to get the filename for
 	 * @return the filename
 	 */
-	static SOURCECONTROL_API FString PackageFilename( const UPackage* InPackage );
+	static FString PackageFilename(const UPackage* InPackage);
 
 	/**
-	 * Helper function to convert package array into filename array.
-	 * @param	InPackages	The package array
-	 * @return an array of filenames
-	 */
-	static SOURCECONTROL_API TArray<FString> PackageFilenames( const TArray<UPackage*>& InPackages );
-
-	/**
 	 * Helper function to convert package name array into a filename array.
 	 * @param	InPackageNames	The package name array
 	 * @return an array of filenames
 	 */
-	UFUNCTION(BlueprintCallable, Category = "Editor Scripting | Editor Source Control Helpers")
-	static SOURCECONTROL_API TArray<FString> PackageFilenames( const TArray<FString>& InPackageNames );
-	
+	static TArray<FString> PackageFilenames(const TArray<FString>& InPackageNames);
+
 	/**
+	 * Helper function to convert package array into filename array.
+	 * @param	InPackages	The package array
+	 * @return an array of filenames
+	 */
+	static TArray<FString> PackageFilenames(const TArray<UPackage*>& InPackages);
+
+	/**
 	 * Helper function to convert a filename array to absolute paths.
 	 * @param	InFileNames	The filename array
 	 * @return an array of filenames, transformed into absolute paths
 	 */
-	UFUNCTION(BlueprintCallable, Category = "Editor Scripting | Editor Source Control Helpers")
-	static SOURCECONTROL_API TArray<FString> AbsoluteFilenames( const TArray<FString>& InFileNames );
+	static TArray<FString> AbsoluteFilenames(const TArray<FString>& InFileNames);
 
 	/**
-	 * Helper function to get a list of files that are unchanged & revert them. This runs synchronous commands.
-	 * @param	InProvider	The provider to use
-	 * @param	InFiles		The files to operate on
-	 */
-	static SOURCECONTROL_API void RevertUnchangedFiles( ISourceControlProvider& InProvider, const TArray<FString>& InFiles );
-
-	/**
 	 * Helper function to annotate a file using a label
 	 * @param	InProvider	The provider to use
 	 * @param	InLabel		The label to use to retrieve the file
@@ -94,54 +415,17 @@
 	 * @param	OutLines	Output array of annotated lines
 	 * @returns true if successful
 	 */
-	static SOURCECONTROL_API bool AnnotateFile( ISourceControlProvider& InProvider, const FString& InLabel, const FString& InFile, TArray<FAnnotationLine>& OutLines );
+	static bool AnnotateFile(ISourceControlProvider& InProvider, const FString& InLabel, const FString& InFile, TArray<FAnnotationLine>& OutLines);
 
 	/**
 	 * Helper function to annotate a file using a changelist/checkin identifier
 	 * @param	InProvider				The provider to use
-	 * @param	InCheckInIdentifier		The changelist/check identifier to use to retrieve the file
+	 * @param	InCheckInIdentifier		The changelist/checkin identifier to use to retrieve the file
 	 * @param	InFile					The file to annotate
 	 * @param	OutLines				Output array of annotated lines
 	 * @returns true if successful
 	 */
-	static SOURCECONTROL_API bool AnnotateFile( ISourceControlProvider& InProvider, int32 InCheckInIdentifier, const FString& InFile, TArray<FAnnotationLine>& OutLines );
+	static bool AnnotateFile(ISourceControlProvider& InProvider, int32 InCheckInIdentifier, const FString& InFile, TArray<FAnnotationLine>& OutLines);
 
 	/**
-	 * Helper function to check out a file
-	 * @param	InFile		The file path to check in
-	 * @return	Success or failure of the checkout operation
-	 */
-	UFUNCTION(BlueprintCallable, Category = "Editor Scripting | Editor Source Control Helpers")
-	static SOURCECONTROL_API bool CheckOutFile( const FString& InFile );
-
-	/**
-	 * Helper function to mark a file for add. Does nothing (and returns true) if the file is already under SC
-	 * @param	InFile		The file path to check in
-	 * @return	Success or failure of the mark for add operation
-	 */
-	UFUNCTION(BlueprintCallable, Category = "Editor Scripting | Editor Source Control Helpers")
-	static SOURCECONTROL_API bool MarkFileForAdd( const FString& InFile );
-
-	/**
-	 * Helper function perform an operation on files in our 'source controlled' directories, handling checkout/add etc.
-	 * @param	InDestFile			The path to the destination file
-	 * @param	InFileDescription	Description of the file to display to the user, e.g. "Text" or "Image"
-	 * @param	OnPostCheckOut		Delegate used for performing operation on files that may need a checkout, but before they are added to source control 
-	 * @param	OutFailReason		Text describing why the operation failed
-	 * @return	Success or failure of the operation
-	 */
-	static SOURCECONTROL_API bool CheckoutOrMarkForAdd( const FString& InDestFile, const FText& InFileDescription, const FOnPostCheckOut& OnPostCheckOut, FText& OutFailReason );
-
-	/**
-	 * Helper function to copy a file into our 'source controlled' directories, handling checkout/add etc.
-	 * @param	InDestFile			The path to the destination file
-	 * @param	InSourceFile		The path to the source file
-	 * @param	InFileDescription	Description of the file to display to the user, e.g. "Text" or "Image"
-	 * @param	OutFailReason		Text describing why the operation failed
-	 * @return	Success or failure of the operation
-	 */
-	UFUNCTION(BlueprintCallable, Category = "Editor Scripting | Editor Source Control Helpers")
-	static SOURCECONTROL_API bool CopyFileUnderSourceControl( const FString& InDestFile, const FString& InSourceFile, const FText& InFileDescription, FText& OutFailReason );
-
-	/**
 	 * Helper function to branch/integrate packages from one location to another
@@ -149,8 +433,22 @@
 	 * @param	SourcePackage		The source package
 	 * @return true if the file packages were successfully branched.
 	 */
-	static SOURCECONTROL_API bool BranchPackage( UPackage* DestPackage, UPackage* SourcePackage );
-};
+	static bool BranchPackage(UPackage* DestPackage, UPackage* SourcePackage );
 
+	/**
+	 * Helper function to get the ini filename for storing source control settings
+	 * @return the filename
+	 */
+	static const FString& GetSettingsIni();
+
+	/**
+	 * Helper function to get the ini filename for storing global source control settings
+	 * @return the filename
+	 */
+	static const FString& GetGlobalSettingsIni();
+
+};  // USourceControlHelpers
+
+
 /** 
  * Helper class that ensures FSourceControl is properly initialized and shutdown by calling Init/Close in
diff --git a/Engine/Source/Developer/SourceControl/Public/SourceControlOperations.h b/Engine/Source/Developer/SourceControl/Public/SourceControlOperations.h
--- a/Engine/Source/Developer/SourceControl/Public/SourceControlOperations.h
+++ b/Engine/Source/Developer/SourceControl/Public/SourceControlOperations.h
@@ -3,25 +3,25 @@
 #pragma once
 
 #include "CoreMinimal.h"
-#include "ISourceControlOperation.h"
+#include "SourceControlOperationBase.h"
 
 #define LOCTEXT_NAMESPACE "SourceControl"
 
 /**
  * Operation used to connect (or test a connection) to source control
  */
-class FConnect : public ISourceControlOperation
+class FConnect : public FSourceControlOperationBase
 {
 public:
 	// ISourceControlOperation interface
-	virtual FName GetName() const override 
-	{ 
-		return "Connect"; 
+	virtual FName GetName() const override
+	{
+		return "Connect";
 	}
 
 	virtual FText GetInProgressString() const override
-	{ 
-		return LOCTEXT("SourceControl_Connecting", "Connecting to source control..."); 
+	{
+		return LOCTEXT("SourceControl_Connecting", "Connecting to source control...");
 	}
 
 	const FString& GetPassword() const
@@ -55,7 +55,7 @@ protected:
 /**
  * Operation used to check files into source control
  */
-class FCheckIn : public ISourceControlOperation
+class FCheckIn : public FSourceControlOperationBase
 {
 public:
 	// ISourceControlOperation interface
@@ -65,8 +65,8 @@ public:
 	}
 
 	virtual FText GetInProgressString() const override
-	{ 
-		return LOCTEXT("SourceControl_CheckIn", "Checking file(s) into Source Control..."); 
+	{
+		return LOCTEXT("SourceControl_CheckIn", "Checking file(s) into Source Control...");
 	}
 
 	void SetDescription( const FText& InDescription )
@@ -100,7 +100,7 @@ protected:
 /**
  * Operation used to check files out of source control
  */
-class FCheckOut : public ISourceControlOperation
+class FCheckOut : public FSourceControlOperationBase
 {
 public:
 	// ISourceControlOperation interface
@@ -110,15 +110,15 @@ public:
 	}
 
 	virtual FText GetInProgressString() const override
-	{ 
-		return LOCTEXT("SourceControl_CheckOut", "Checking file(s) out of Source Control..."); 
+	{
+		return LOCTEXT("SourceControl_CheckOut", "Checking file(s) out of Source Control...");
 	}
 };
 
 /**
  * Operation used to mark files for add in source control
  */
-class FMarkForAdd : public ISourceControlOperation
+class FMarkForAdd : public FSourceControlOperationBase
 {
 public:
 	// ISourceControlOperation interface
@@ -128,15 +128,15 @@ public:
 	}
 
 	virtual FText GetInProgressString() const override
-	{ 
-		return LOCTEXT("SourceControl_Add", "Adding file(s) to Source Control..."); 
+	{
+		return LOCTEXT("SourceControl_Add", "Adding file(s) to Source Control...");
 	}
 };
 
 /**
  * Operation used to mark files for delete in source control
  */
-class FDelete : public ISourceControlOperation
+class FDelete : public FSourceControlOperationBase
 {
 public:
 	// ISourceControlOperation interface
@@ -146,15 +146,15 @@ public:
 	}
 
 	virtual FText GetInProgressString() const override
-	{ 
-		return LOCTEXT("SourceControl_Delete", "Deleting file(s) from Source Control..."); 
+	{
+		return LOCTEXT("SourceControl_Delete", "Deleting file(s) from Source Control...");
 	}
 };
 
 /**
  * Operation used to revert changes made back to the state they are in source control
  */
-class FRevert : public ISourceControlOperation
+class FRevert : public FSourceControlOperationBase
 {
 public:
 	// ISourceControlOperation interface
@@ -164,15 +164,15 @@ public:
 	}
 
 	virtual FText GetInProgressString() const override
-	{ 
-		return LOCTEXT("SourceControl_Revert", "Reverting file(s) in Source Control..."); 
-	}	
+	{
+		return LOCTEXT("SourceControl_Revert", "Reverting file(s) in Source Control...");
+	}
 };
 
 /**
  * Operation used to sync files to the state they are in source control
  */
-class FSync : public ISourceControlOperation
+class FSync : public FSourceControlOperationBase
 {
 public:
 	// ISourceControlOperation interface
@@ -182,9 +182,9 @@ public:
 	}
 
 	virtual FText GetInProgressString() const override
-	{ 
-		return LOCTEXT("SourceControl_Sync", "Syncing file(s) from source control..."); 
-	}	
+	{
+		return LOCTEXT("SourceControl_Sync", "Syncing file(s) from source control...");
+	}
 
 	void SetRevision( int32 InRevisionNumber )
 	{
@@ -199,7 +199,7 @@ protected:
 /**
  * Operation used to update the source control status of files
  */
-class FUpdateStatus : public ISourceControlOperation
+class FUpdateStatus : public FSourceControlOperationBase
 {
 public:
 	FUpdateStatus()
@@ -216,9 +216,9 @@ public:
 	}
 
 	virtual FText GetInProgressString() const override
-	{ 
-		return LOCTEXT("SourceControl_Update", "Updating file(s) source control status..."); 
-	}	
+	{
+		return LOCTEXT("SourceControl_Update", "Updating file(s) source control status...");
+	}
 
 	void SetUpdateHistory( bool bInUpdateHistory )
 	{
@@ -277,7 +277,7 @@ protected:
 /**
  * Operation used to copy a file or directory from one location to another
  */
-class FCopy : public ISourceControlOperation
+class FCopy : public FSourceControlOperationBase
 {
 public:
 	// ISourceControlOperation interface
@@ -287,9 +287,9 @@ public:
 	}
 
 	virtual FText GetInProgressString() const override
-	{ 
-		return LOCTEXT("SourceControl_Copy", "Copying file(s) in Source Control..."); 
-	}	
+	{
+		return LOCTEXT("SourceControl_Copy", "Copying file(s) in Source Control...");
+	}
 
 	void SetDestination(const FString& InDestination)
 	{
@@ -309,7 +309,7 @@ protected:
 /**
  * Operation used to resolve a file that is in a conflicted state.
  */
-class FResolve : public ISourceControlOperation
+class FResolve : public FSourceControlOperationBase
 {
 public:
 	// ISourceControlOperation interface
diff --git a/Engine/Source/Editor/ContentBrowser/Private/AssetContextMenu.cpp b/Engine/Source/Editor/ContentBrowser/Private/AssetContextMenu.cpp
--- a/Engine/Source/Editor/ContentBrowser/Private/AssetContextMenu.cpp
+++ b/Engine/Source/Editor/ContentBrowser/Private/AssetContextMenu.cpp
@@ -10,6 +10,7 @@
 #include "HAL/PlatformApplicationMisc.h"
 #include "HAL/FileManager.h"
 #include "Misc/ScopedSlowTask.h"
+#include "UObject/MetaData.h"
 #include "UObject/UObjectIterator.h"
 #include "Widgets/SBoxPanel.h"
 #include "Widgets/SWindow.h"
@@ -17,6 +18,7 @@
 #include "Widgets/Text/STextBlock.h"
 #include "Framework/MultiBox/MultiBoxBuilder.h"
 #include "Widgets/Input/SButton.h"
+#include "Widgets/Input/SMultiLineEditableTextBox.h"
 #include "EditorStyleSet.h"
 #include "EditorReimportHandler.h"
 #include "Components/ActorComponent.h"
@@ -24,9 +26,9 @@
 #include "UnrealClient.h"
 #include "Materials/MaterialFunctionInstance.h"
 #include "Materials/Material.h"
-#include "ISourceControlOperation.h"
 #include "SourceControlOperations.h"
 #include "ISourceControlModule.h"
+#include "SourceControlHelpers.h"
 #include "Settings/EditorExperimentalSettings.h"
 #include "Materials/MaterialInstanceConstant.h"
 #include "FileHelpers.h"
@@ -36,6 +38,7 @@
 #include "ContentBrowserUtils.h"
 #include "SAssetView.h"
 #include "ContentBrowserModule.h"
+#include "Dialogs/Dialogs.h"
 
 #include "ObjectTools.h"
 #include "PackageTools.h"
@@ -354,10 +357,10 @@ void FAssetContextMenu::MakeAssetActionsSubMenu(FMenuBuilder& MenuBuilder)
 		LOCTEXT("CreateBlueprintUsingTooltip", "Create a new Blueprint and add this asset to it"),
 		FSlateIcon(FEditorStyle::GetStyleSetName(), "LevelEditor.CreateClassBlueprint"),
 		FUIAction(
-		FExecuteAction::CreateSP(this, &FAssetContextMenu::ExecuteCreateBlueprintUsing),
-		FCanExecuteAction::CreateSP(this, &FAssetContextMenu::CanExecuteCreateBlueprintUsing)
-			)
-		);
+			FExecuteAction::CreateSP(this, &FAssetContextMenu::ExecuteCreateBlueprintUsing),
+			FCanExecuteAction::CreateSP(this, &FAssetContextMenu::CanExecuteCreateBlueprintUsing)
+		)
+	);
 
 	// Capture Thumbnail
 	FAssetToolsModule& AssetToolsModule = FModuleManager::LoadModuleChecked<FAssetToolsModule>("AssetTools");
@@ -368,10 +371,10 @@ void FAssetContextMenu::MakeAssetActionsSubMenu(FMenuBuilder& MenuBuilder)
 			LOCTEXT("CaptureThumbnailTooltip", "Captures a thumbnail from the active viewport."),
 			FSlateIcon(FEditorStyle::GetStyleSetName(), "ContentBrowser.AssetActions.CreateThumbnail"),
 			FUIAction(
-			FExecuteAction::CreateSP(this, &FAssetContextMenu::ExecuteCaptureThumbnail),
-			FCanExecuteAction::CreateSP(this, &FAssetContextMenu::CanExecuteCaptureThumbnail)
-				)
-			);
+				FExecuteAction::CreateSP(this, &FAssetContextMenu::ExecuteCaptureThumbnail),
+				FCanExecuteAction::CreateSP(this, &FAssetContextMenu::CanExecuteCaptureThumbnail)
+			)
+		);
 	}
 
 	// Clear Thumbnail
@@ -507,6 +510,17 @@ void FAssetContextMenu::MakeAssetActionsSubMenu(FMenuBuilder& MenuBuilder)
 				);
 		}
 
+		// Create Metadata menu
+		MenuBuilder.AddMenuEntry(
+			LOCTEXT("ShowAssetMetaData", "Show Meta-data"),
+			LOCTEXT("ShowAssetMetaDataTooltip", "Show the asset meta-data dialog."),
+			FSlateIcon(),
+			FUIAction(
+				FExecuteAction::CreateSP(this, &FAssetContextMenu::ExecuteShowAssetMetaData),
+				FCanExecuteAction::CreateSP(this, &FAssetContextMenu::CanExecuteShowAssetMetaData)
+			)
+		);
+
 		// Chunk actions
 		if (GetDefault<UEditorExperimentalSettings>()->bContextMenuChunkAssignments)
 		{
@@ -1742,6 +1756,51 @@ void FAssetContextMenu::ExecutePropertyMatrix()
 	}
 }
 
+void FAssetContextMenu::ExecuteShowAssetMetaData()
+{
+	bool bFirstAsset = true;
+	FString OutputValue;
+
+	for (const FAssetData& AssetData : SelectedAssets)
+	{
+		UObject* Asset = AssetData.GetAsset();
+		if (Asset)
+		{
+			if (!bFirstAsset)
+			{
+				OutputValue.Append(TEXT("\n\n"));
+			}
+			bFirstAsset = false;
+
+			OutputValue.Append(*Asset->GetPathName());
+
+			TMap<FName, FString>* TagValues = UMetaData::GetMapForObject(Asset);
+			if (TagValues)
+			{
+				for (TMap<FName, FString>::TIterator It(*TagValues); It; ++It)
+				{
+					OutputValue.Append(TEXT("\n  "));
+					OutputValue.Append(It.Key().ToString());
+					OutputValue.Append(TEXT(": "));
+					OutputValue.Append(It.Value());
+				}
+			}
+			else
+			{
+				OutputValue.Append(TEXT("\n  "));
+				OutputValue.Append(NSLOCTEXT("AssetContextMenu", "ShowMetaData", "No meta-data found.").ToString());
+			}
+		}
+	}
+
+	SGenericDialogWidget::OpenDialog(
+		NSLOCTEXT("AssetContextMenu", "ShowMetaData", "Show Meta-data"),
+		SNew(SMultiLineEditableTextBox)
+			.Text(FText::FromString(OutputValue))
+			.IsReadOnly(true)
+	);
+}
+
 void FAssetContextMenu::ExecuteEditAsset()
 {
 	TMap<UClass*, TArray<UObject*> > SelectedAssetsByClass;
@@ -2413,6 +2472,24 @@ FText FAssetContextMenu::GetExecutePropertyMatrixTooltip() const
 	return ResultTooltip;
 }
 
+bool FAssetContextMenu::CanExecuteShowAssetMetaData() const
+{
+	TArray<UObject*> ObjectsForPropertiesMenu;
+	const bool SkipRedirectors = true;
+	GetSelectedAssets(ObjectsForPropertiesMenu, SkipRedirectors);
+
+	bool bResult = false;
+	for (const UObject* Asset : ObjectsForPropertiesMenu)
+	{
+		if (Asset && UMetaData::GetMapForObject(Asset))
+		{
+			bResult = true;
+			break;
+		}
+	}
+	return bResult;
+}
+
 bool FAssetContextMenu::CanExecuteDuplicate() const
 {
 	const TArray< FAssetData > AssetViewSelectedAssets = AssetView.Pin()->GetSelectedAssets();
diff --git a/Engine/Source/Editor/ContentBrowser/Private/AssetContextMenu.h b/Engine/Source/Editor/ContentBrowser/Private/AssetContextMenu.h
--- a/Engine/Source/Editor/ContentBrowser/Private/AssetContextMenu.h
+++ b/Engine/Source/Editor/ContentBrowser/Private/AssetContextMenu.h
@@ -190,6 +190,9 @@ private:
 	/** Handler for when "Property Matrix..." is selected */
 	void ExecutePropertyMatrix();
 
+	/** Handler for when "Show MetaData" is selected */
+	void ExecuteShowAssetMetaData();
+
 	/** Handler for when "Edit Asset" is selected */
 	void ExecuteEditAsset();
 
@@ -302,6 +305,9 @@ private:
 	bool CanExecutePropertyMatrix(FText& OutErrorMessage) const;
 	bool CanExecutePropertyMatrix() const;
 
+	/** Handler to check to see if a "Show MetaData" command is allowed */
+	bool CanExecuteShowAssetMetaData() const;
+
 	/** Handler to check to see if a "Remove from collection" command is allowed */
 	bool CanExecuteRemoveFromCollection() const;
 
diff --git a/Engine/Source/Editor/ContentBrowser/Private/FrontendFilters.cpp b/Engine/Source/Editor/ContentBrowser/Private/FrontendFilters.cpp
--- a/Engine/Source/Editor/ContentBrowser/Private/FrontendFilters.cpp
+++ b/Engine/Source/Editor/ContentBrowser/Private/FrontendFilters.cpp
@@ -7,6 +7,8 @@
 #include "Framework/MultiBox/MultiBoxBuilder.h"
 #include "Widgets/Input/SEditableTextBox.h"
 #include "ISourceControlModule.h"
+#include "SourceControlHelpers.h"
+#include "SourceControlOperations.h"
 #include "Editor.h"
 #include "AssetToolsModule.h"
 #include "ICollectionManager.h"
diff --git a/Engine/Source/Editor/ContentBrowser/Private/SAssetView.cpp b/Engine/Source/Editor/ContentBrowser/Private/SAssetView.cpp
--- a/Engine/Source/Editor/ContentBrowser/Private/SAssetView.cpp
+++ b/Engine/Source/Editor/ContentBrowser/Private/SAssetView.cpp
@@ -2628,7 +2628,9 @@ void SAssetView::OnCollectionUpdated( const FCollectionNameType& Collection )
 void SAssetView::OnAssetRenamed(const FAssetData& AssetData, const FString& OldObjectPath)
 {
 	// Remove the old asset, if it exists
-	RemoveAssetByPath( FName( *OldObjectPath ) );
+	FName OldObjectPackageName = *OldObjectPath;
+	RemoveAssetByPath( OldObjectPackageName );
+	RecentlyAddedAssets.RemoveAllSwap( [&](const FAssetData& AssetData) { return AssetData.ObjectPath == OldObjectPackageName; } );
 
 	// Add the new asset, if it should be in the cached list
 	OnAssetAdded( AssetData );
diff --git a/Engine/Source/Editor/Kismet/Private/SBlueprintRevisionMenu.cpp b/Engine/Source/Editor/Kismet/Private/SBlueprintRevisionMenu.cpp
--- a/Engine/Source/Editor/Kismet/Private/SBlueprintRevisionMenu.cpp
+++ b/Engine/Source/Editor/Kismet/Private/SBlueprintRevisionMenu.cpp
@@ -12,6 +12,8 @@
 #include "EditorStyleSet.h"
 #include "Engine/Blueprint.h"
 #include "ISourceControlModule.h"
+#include "ISourceControlRevision.h"
+#include "SourceControlHelpers.h"
 #include "IAssetTypeActions.h"
 #include "Widgets/Images/SThrobber.h"
 #include "Kismet2/BlueprintEditorUtils.h"
diff --git a/Engine/Source/Editor/Kismet/Public/SBlueprintRevisionMenu.h b/Engine/Source/Editor/Kismet/Public/SBlueprintRevisionMenu.h
--- a/Engine/Source/Editor/Kismet/Public/SBlueprintRevisionMenu.h
+++ b/Engine/Source/Editor/Kismet/Public/SBlueprintRevisionMenu.h
@@ -8,6 +8,7 @@
 #include "Widgets/DeclarativeSyntaxSupport.h"
 #include "Widgets/SCompoundWidget.h"
 #include "ISourceControlProvider.h"
+#include "SourceControlOperations.h"
 
 class SVerticalBox;
 class UBlueprint;
diff --git a/Engine/Source/Editor/LevelEditor/Private/LevelEditorActions.cpp b/Engine/Source/Editor/LevelEditor/Private/LevelEditorActions.cpp
--- a/Engine/Source/Editor/LevelEditor/Private/LevelEditorActions.cpp
+++ b/Engine/Source/Editor/LevelEditor/Private/LevelEditorActions.cpp
@@ -18,6 +18,7 @@
 #include "Materials/Material.h"
 #include "Editor/EditorPerProjectUserSettings.h"
 #include "ISourceControlModule.h"
+#include "SourceControlHelpers.h"
 #include "Editor/UnrealEdEngine.h"
 #include "Settings/EditorExperimentalSettings.h"
 #include "Factories/BlueprintFactory.h"
@@ -2216,7 +2217,7 @@ void FLevelEditorActionCallbacks::CheckOutProjectSettingsConfig( )
 	FString ConfigFilePath = FPaths::ConvertRelativePathToFull(FString::Printf(TEXT("%sDefaultEngine.ini"), *FPaths::SourceConfigDir()));
 	if(ISourceControlModule::Get().IsEnabled())
 	{
-		SourceControlHelpers::CheckOutFile(ConfigFilePath);
+		USourceControlHelpers::CheckOutOrAddFile(ConfigFilePath);
 	}
 	else
 	{
diff --git a/Engine/Source/Editor/MovieSceneTools/Private/MovieSceneToolHelpers.cpp b/Engine/Source/Editor/MovieSceneTools/Private/MovieSceneToolHelpers.cpp
--- a/Engine/Source/Editor/MovieSceneTools/Private/MovieSceneToolHelpers.cpp
+++ b/Engine/Source/Editor/MovieSceneTools/Private/MovieSceneToolHelpers.cpp
@@ -17,6 +17,7 @@
 #include "EditorDirectories.h"
 #include "Sections/MovieSceneFloatSection.h"
 #include "Tracks/MovieSceneFloatTrack.h"
+#include "Tracks/MovieSceneCameraCutTrack.h"
 #include "Sections/MovieScene3DTransformSection.h"
 #include "Tracks/MovieScene3DTransformTrack.h"
 #include "Sections/MovieSceneCinematicShotSection.h"
@@ -1171,6 +1172,79 @@ void ImportFBXCamera(UnFbx::FFbxImporter* FbxImporter, UMovieScene* InMovieScene
 	}
 }
 
+FGuid FindCameraGuid(FbxCamera* Camera, TMap<FGuid, FString>& InObjectBindingMap)
+{
+	for (auto& Pair : InObjectBindingMap)
+	{
+		if (FCString::Strcmp(*Pair.Value, UTF8_TO_TCHAR(Camera->GetName())) == 0)
+		{
+			return Pair.Key;
+		}
+	}
+	return FGuid();
+}
+
+UMovieSceneCameraCutTrack* GetCameraCutTrack(UMovieScene* InMovieScene)
+{
+	// Get the camera cut
+	UMovieSceneTrack* CameraCutTrack = InMovieScene->GetCameraCutTrack();
+	if (CameraCutTrack == nullptr)
+	{
+		InMovieScene->Modify();
+		CameraCutTrack = InMovieScene->AddCameraCutTrack(UMovieSceneCameraCutTrack::StaticClass());
+	}
+	return CastChecked<UMovieSceneCameraCutTrack>(CameraCutTrack);
+}
+
+void ImportCameraCut(UnFbx::FFbxImporter* FbxImporter, UMovieScene* InMovieScene, ISequencer& InSequencer, TMap<FGuid, FString>& InObjectBindingMap)
+{
+	// The camera switcher refers to the fbx camera
+	const UMovieSceneUserImportFBXSettings* ImportFBXSettings = GetDefault<UMovieSceneUserImportFBXSettings>();
+	if (!ImportFBXSettings->bCreateCameras)
+	{
+		return;
+	}
+	// Find a camera switcher
+	FbxCameraSwitcher* CameraSwitcher = FbxImporter->Scene->GlobalCameraSettings().GetCameraSwitcher();
+	if (CameraSwitcher == nullptr)
+	{
+		return;
+	}
+	// Get the animation layer
+	FbxAnimStack* AnimStack = FbxImporter->Scene->GetMember<FbxAnimStack>(0);
+	if (AnimStack == nullptr)
+	{
+		return;
+	}
+	FbxAnimLayer* AnimLayer = AnimStack->GetMember<FbxAnimLayer>(0);
+	if (AnimLayer == nullptr)
+	{
+		return;
+	}
+
+	// The camera switcher camera index refer to depth-first found order of the camera in the FBX
+	TArray<FbxCamera*> AllCameras;
+	GetCameras(FbxImporter->Scene->GetRootNode(), AllCameras);
+
+	UMovieSceneCameraCutTrack* CameraCutTrack = GetCameraCutTrack(InMovieScene);
+	FFrameRate FrameRate = CameraCutTrack->GetTypedOuter<UMovieScene>()->GetFrameResolution();
+
+	FbxAnimCurve* AnimCurve = CameraSwitcher->CameraIndex.GetCurve(AnimLayer);
+	for (int i = 0; i < AnimCurve->KeyGetCount(); ++i)
+	{
+		FbxAnimCurveKey key = AnimCurve->KeyGet(i);
+		int value = (int)key.GetValue() - 1;
+		if (value >= 0 && value < AllCameras.Num())
+		{
+			FGuid CameraGuid = FindCameraGuid(AllCameras[value], InObjectBindingMap);
+			if (CameraGuid != FGuid())
+			{
+				CameraCutTrack->AddNewCameraCut(FMovieSceneObjectBindingID(CameraGuid, MovieSceneSequenceID::Root), (((float)key.GetTime().GetSecondDouble()) * FrameRate).FloorToFrame());
+			}
+		}
+	}
+	InSequencer.NotifyMovieSceneDataChanged(EMovieSceneDataChangeType::MovieSceneStructureItemAdded);
+}
 
 class SMovieSceneImportFBXSettings : public SCompoundWidget, public FGCObject
 {
@@ -1277,6 +1351,9 @@ private:
 		TArray<FString> AllNodeNames;
 		CurveAPI.GetAllNodeNameArray(AllNodeNames);
 
+		// Import a camera cut track, do it after populating curve data ensure only one animation layer, if any
+		ImportCameraCut(FbxImporter, MovieScene, *Sequencer, ObjectBindingMap);
+
 		for (FString NodeName : AllNodeNames)
 		{
 			ImportFBXNode(NodeName, CurveAPI, MovieScene, *Sequencer, ObjectBindingMap);
diff --git a/Engine/Source/Editor/SourceControlWindows/Private/SSourceControlHistory.cpp b/Engine/Source/Editor/SourceControlWindows/Private/SSourceControlHistory.cpp
--- a/Engine/Source/Editor/SourceControlWindows/Private/SSourceControlHistory.cpp
+++ b/Engine/Source/Editor/SourceControlWindows/Private/SSourceControlHistory.cpp
@@ -7,6 +7,8 @@
 #include "SourceControlOperations.h"
 #include "ISourceControlRevision.h"
 #include "SourceControlWindows.h"
+#include "SourceControlHelpers.h"
+#include "ISourceControlModule.h"
 #include "Modules/ModuleManager.h"
 #include "UObject/Object.h"
 #include "UObject/Package.h"
@@ -41,7 +43,6 @@
 #include "Widgets/Views/STreeView.h"
 #include "Framework/Docking/TabManager.h"
 #include "EditorStyleSet.h"
-#include "ISourceControlModule.h"
 
 #include "IAssetTools.h"
 #include "IAssetTypeActions.h"
diff --git a/Engine/Source/Editor/SourceControlWindows/Private/SSourceControlRevert.cpp b/Engine/Source/Editor/SourceControlWindows/Private/SSourceControlRevert.cpp
--- a/Engine/Source/Editor/SourceControlWindows/Private/SSourceControlRevert.cpp
+++ b/Engine/Source/Editor/SourceControlWindows/Private/SSourceControlRevert.cpp
@@ -5,6 +5,8 @@
 #include "ISourceControlOperation.h"
 #include "SourceControlOperations.h"
 #include "SourceControlWindows.h"
+#include "SourceControlHelpers.h"
+#include "ISourceControlModule.h"
 #include "UObject/Package.h"
 #include "Misc/PackageName.h"
 #include "Layout/Visibility.h"
@@ -26,7 +28,6 @@
 #include "Widgets/Views/SListView.h"
 #include "Widgets/Input/SCheckBox.h"
 #include "EditorStyleSet.h"
-#include "ISourceControlModule.h"
 #include "PackageTools.h"
 #include "Settings/EditorExperimentalSettings.h"
 #include "Framework/Notifications/NotificationManager.h"
diff --git a/Engine/Source/Editor/SourceControlWindows/Private/SSourceControlSubmit.cpp b/Engine/Source/Editor/SourceControlWindows/Private/SSourceControlSubmit.cpp
--- a/Engine/Source/Editor/SourceControlWindows/Private/SSourceControlSubmit.cpp
+++ b/Engine/Source/Editor/SourceControlWindows/Private/SSourceControlSubmit.cpp
@@ -1,14 +1,13 @@
 // Copyright 1998-2018 Epic Games, Inc. All Rights Reserved.
 
 #include "SSourceControlSubmit.h"
-#include "Misc/MessageDialog.h"
 #include "ISourceControlOperation.h"
 #include "SourceControlOperations.h"
 #include "ISourceControlProvider.h"
-#include "SourceControlWindows.h"
+#include "ISourceControlModule.h"
+#include "SourceControlHelpers.h"
 #include "Modules/ModuleManager.h"
 #include "Widgets/SWindow.h"
-#include "Framework/Application/SlateApplication.h"
 #include "Widgets/Layout/SSpacer.h"
 #include "Widgets/Layout/SWrapBox.h"
 #include "Widgets/Images/SImage.h"
@@ -17,19 +16,10 @@
 #include "Widgets/Input/SMultiLineEditableTextBox.h"
 #include "Widgets/Input/SButton.h"
 #include "Widgets/Notifications/SErrorText.h"
-#include "Widgets/Views/SListView.h"
 #include "Widgets/Input/SCheckBox.h"
 #include "EditorStyleSet.h"
-#include "ISourceControlModule.h"
-#include "Logging/TokenizedMessage.h"
-#include "Logging/MessageLog.h"
-#include "Framework/Notifications/NotificationManager.h"
-#include "Widgets/Notifications/SNotificationList.h"
-#include "FileHelpers.h"
 
 
-IMPLEMENT_MODULE( FDefaultModuleImpl, SourceControlWindows );
-
 #if SOURCE_CONTROL_WITH_SLATE
 
 #define LOCTEXT_NAMESPACE "SSourceControlSubmit"
@@ -465,305 +455,6 @@ void SSourceControlSubmitWidget::SortTree()
 }
 
 
-TWeakPtr<SNotificationItem> FSourceControlWindows::ChoosePackagesToCheckInNotification;
-
-void FSourceControlWindows::ChoosePackagesToCheckInCompleted(const TArray<UPackage*>& LoadedPackages, const TArray<FString>& PackageNames, const TArray<FString>& ConfigFiles)
-{
-	if (ChoosePackagesToCheckInNotification.IsValid())
-	{
-		ChoosePackagesToCheckInNotification.Pin()->ExpireAndFadeout();
-	}
-	ChoosePackagesToCheckInNotification.Reset();
-
-	// Prompt the user to ask if they would like to first save any dirty packages they are trying to check-in
-	const FEditorFileUtils::EPromptReturnCode UserResponse = FEditorFileUtils::PromptForCheckoutAndSave(LoadedPackages, true, true);
-
-	// If the user elected to save dirty packages, but one or more of the packages failed to save properly OR if the user
-	// canceled out of the prompt, don't follow through on the check-in process
-	const bool bShouldProceed = (UserResponse == FEditorFileUtils::EPromptReturnCode::PR_Success || UserResponse == FEditorFileUtils::EPromptReturnCode::PR_Declined);
-	if (bShouldProceed)
-	{
-		TArray<FString> PendingDeletePaths;
-		PendingDeletePaths.Add(FPaths::ConvertRelativePathToFull(FPaths::EngineContentDir()));
-		PendingDeletePaths.Add(FPaths::ConvertRelativePathToFull(FPaths::ProjectContentDir()));
-		PendingDeletePaths.Add(FPaths::ConvertRelativePathToFull(FPaths::ProjectConfigDir()));
-		PendingDeletePaths.Add(FPaths::ConvertRelativePathToFull(FPaths::GetProjectFilePath()));
-
-		const bool bUseSourceControlStateCache = true;
-		PromptForCheckin(bUseSourceControlStateCache, PackageNames, PendingDeletePaths, ConfigFiles);
-	}
-	else
-	{
-		// If a failure occurred, alert the user that the check-in was aborted. This warning shouldn't be necessary if the user cancelled
-		// from the dialog, because they obviously intended to cancel the whole operation.
-		if (UserResponse == FEditorFileUtils::EPromptReturnCode::PR_Failure)
-		{
-			FMessageDialog::Open(EAppMsgType::Ok, NSLOCTEXT("UnrealEd", "SCC_Checkin_Aborted", "Check-in aborted as a result of save failure."));
-		}
-	}
-}
-
-void FSourceControlWindows::ChoosePackagesToCheckInCancelled(FSourceControlOperationRef InOperation)
-{
-	ISourceControlProvider& SourceControlProvider = ISourceControlModule::Get().GetProvider();
-	SourceControlProvider.CancelOperation(InOperation);
-
-	if (ChoosePackagesToCheckInNotification.IsValid())
-	{
-		ChoosePackagesToCheckInNotification.Pin()->ExpireAndFadeout();
-	}
-	ChoosePackagesToCheckInNotification.Reset();
-}
-
-void FSourceControlWindows::ChoosePackagesToCheckInCallback(const FSourceControlOperationRef& InOperation, ECommandResult::Type InResult)
-{
-	if (ChoosePackagesToCheckInNotification.IsValid())
-	{
-		ChoosePackagesToCheckInNotification.Pin()->ExpireAndFadeout();
-	}
-	ChoosePackagesToCheckInNotification.Reset();
-
-	if (InResult == ECommandResult::Succeeded)
-	{
-		// Get a list of all the checked out packages
-		TArray<FString> PackageNames;
-		TArray<UPackage*> LoadedPackages;
-		TMap<FString, FSourceControlStatePtr> PackageStates;
-		FEditorFileUtils::FindAllSubmittablePackageFiles(PackageStates, true);
-
-		for (TMap<FString, FSourceControlStatePtr>::TConstIterator PackageIter(PackageStates); PackageIter; ++PackageIter)
-		{
-			const FString PackageName = *PackageIter.Key();
-			const FSourceControlStatePtr CurPackageSCCState = PackageIter.Value();
-
-			UPackage* Package = FindPackage(nullptr, *PackageName);
-			if (Package != nullptr)
-			{
-				LoadedPackages.Add(Package);
-			}
-			
-			PackageNames.Add(PackageName);
-		}
-
-		// Get a list of all the checked out config files
-		TMap<FString, FSourceControlStatePtr> ConfigFileStates;
-		TArray<FString> ConfigFilesToSubmit;
-		FEditorFileUtils::FindAllSubmittableConfigFiles(ConfigFileStates);
-		for (TMap<FString, FSourceControlStatePtr>::TConstIterator It(ConfigFileStates); It; ++It)
-		{
-			ConfigFilesToSubmit.Add(It.Key());
-		}
-
-		ChoosePackagesToCheckInCompleted(LoadedPackages, PackageNames, ConfigFilesToSubmit);
-	}
-	else if (InResult == ECommandResult::Failed)
-	{
-		FMessageLog EditorErrors("EditorErrors");
-		EditorErrors.Warning(LOCTEXT("CheckInOperationFailed", "Failed checking source control status!"));
-		EditorErrors.Notify();
-	}
-}
-
-void FSourceControlWindows::ChoosePackagesToCheckIn()
-{
-	if (ISourceControlModule::Get().IsEnabled())
-	{
-		if (ISourceControlModule::Get().GetProvider().IsAvailable())
-		{
-			// make sure we update the SCC status of all packages (this could take a long time, so we will run it as a background task)
-			TArray<FString> Filenames;
-			Filenames.Add(FPaths::ConvertRelativePathToFull(FPaths::EngineContentDir()));
-			Filenames.Add(FPaths::ConvertRelativePathToFull(FPaths::ProjectContentDir()));
-			Filenames.Add(FPaths::ConvertRelativePathToFull(FPaths::ProjectConfigDir()));
-			Filenames.Add(FPaths::ConvertRelativePathToFull(FPaths::GetProjectFilePath()));
-
-			ISourceControlProvider& SourceControlProvider = ISourceControlModule::Get().GetProvider();
-			FSourceControlOperationRef Operation = ISourceControlOperation::Create<FUpdateStatus>();
-			StaticCastSharedRef<FUpdateStatus>(Operation)->SetCheckingAllFiles(false);
-			SourceControlProvider.Execute(Operation, Filenames, EConcurrency::Asynchronous, FSourceControlOperationComplete::CreateStatic(&FSourceControlWindows::ChoosePackagesToCheckInCallback));
-
-			if (ChoosePackagesToCheckInNotification.IsValid())
-			{
-				ChoosePackagesToCheckInNotification.Pin()->ExpireAndFadeout();
-			}
-
-			FNotificationInfo Info(LOCTEXT("ChooseAssetsToCheckInIndicator", "Checking for assets to check in..."));
-			Info.bFireAndForget = false;
-			Info.ExpireDuration = 0.0f;
-			Info.FadeOutDuration = 1.0f;
-
-			if (SourceControlProvider.CanCancelOperation(Operation))
-			{
-				Info.ButtonDetails.Add(FNotificationButtonInfo(
-					LOCTEXT("ChoosePackagesToCheckIn_CancelButton", "Cancel"),
-					LOCTEXT("ChoosePackagesToCheckIn_CancelButtonTooltip", "Cancel the check in operation."),
-					FSimpleDelegate::CreateStatic(&FSourceControlWindows::ChoosePackagesToCheckInCancelled, Operation)
-					));
-			}
-
-			ChoosePackagesToCheckInNotification = FSlateNotificationManager::Get().AddNotification(Info);
-
-			if (ChoosePackagesToCheckInNotification.IsValid())
-			{
-				ChoosePackagesToCheckInNotification.Pin()->SetCompletionState(SNotificationItem::CS_Pending);
-			}
-		}
-		else
-		{
-			FMessageLog EditorErrors("EditorErrors");
-			EditorErrors.Warning(LOCTEXT("NoSCCConnection", "No connection to source control available!"))
-				->AddToken(FDocumentationToken::Create(TEXT("Engine/UI/SourceControl")));
-			EditorErrors.Notify();
-		}
-	}
-}
-
-bool FSourceControlWindows::CanChoosePackagesToCheckIn()
-{
-	return !ChoosePackagesToCheckInNotification.IsValid();
-}
-
-
-bool FSourceControlWindows::PromptForCheckin(bool bUseSourceControlStateCache, const TArray<FString>& InPackageNames, const TArray<FString>& InPendingDeletePaths, const TArray<FString>& InConfigFiles)
-{
-	ISourceControlProvider& SourceControlProvider = ISourceControlModule::Get().GetProvider();
-
-	bool bCheckInSuccess = true;
-
-	// Get filenames for packages and config to be checked in
-	TArray<FString> AllFiles = SourceControlHelpers::PackageFilenames(InPackageNames);
-	AllFiles.Append(InConfigFiles);
-
-	// Prepare a list of files to have their states updated
-	if (!bUseSourceControlStateCache)
-	{
-		TArray<FString> UpdateRequest;
-		UpdateRequest.Append(AllFiles);
-
-		// If there are pending delete paths to update, add them here.
-		UpdateRequest.Append(InPendingDeletePaths);
-
-		// Force an update on everything that's been requested
-		if (UpdateRequest.Num() > 0)
-		{
-			SourceControlProvider.Execute(ISourceControlOperation::Create<FUpdateStatus>(), UpdateRequest);
-		}
-	}
-
-	// Get file status of packages and config
-	TArray<FSourceControlStateRef> States;
-	SourceControlProvider.GetState(AllFiles, States, EStateCacheUsage::Use);
-
-	if (InPendingDeletePaths.Num() > 0)
-	{
-		// Get any files pending delete
-		TArray<FSourceControlStateRef> PendingDeleteItems = SourceControlProvider.GetCachedStateByPredicate(
-			[](const FSourceControlStateRef& State) { return State->IsDeleted(); }
-		);
-
-		// And append them to the list
-		States.Append(PendingDeleteItems);
-	}
-
-	if (States.Num())
-	{
-		TSharedRef<SWindow> NewWindow = SNew(SWindow)
-			.Title(NSLOCTEXT("SourceControl.SubmitWindow", "Title", "Submit Files"))
-			.SizingRule(ESizingRule::UserSized)
-			.ClientSize(FVector2D(600, 400))
-			.SupportsMaximize(true)
-			.SupportsMinimize(false);
-
-		TSharedRef<SSourceControlSubmitWidget> SourceControlWidget = 
-			SNew(SSourceControlSubmitWidget)
-			.ParentWindow(NewWindow)
-			.Items(States);
-
-		NewWindow->SetContent(
-			SourceControlWidget
-			);
-
-		FSlateApplication::Get().AddModalWindow(NewWindow, NULL);
-
-		if (SourceControlWidget->GetResult() == ESubmitResults::SUBMIT_ACCEPTED)
-		{
-			//Get description from the dialog
-			FChangeListDescription Description;
-			SourceControlWidget->FillChangeListDescription(Description);
-
-			//revert all unchanged files that were submitted
-			if ( Description.FilesForSubmit.Num() > 0 )
-			{
-				SourceControlHelpers::RevertUnchangedFiles(SourceControlProvider, Description.FilesForSubmit);
-
-				//make sure all files are still checked out
-				for (int32 VerifyIndex = Description.FilesForSubmit.Num()-1; VerifyIndex >= 0; --VerifyIndex)
-				{
-					FSourceControlStatePtr SourceControlState = SourceControlProvider.GetState(Description.FilesForSubmit[VerifyIndex], EStateCacheUsage::Use);
-					if( SourceControlState.IsValid() && !SourceControlState->IsCheckedOut() && !SourceControlState->IsAdded() && !SourceControlState->IsDeleted() )
-					{
-						Description.FilesForSubmit.RemoveAt(VerifyIndex);
-					}
-				}
-			}
-
-			TArray<FString> CombinedFileList = Description.FilesForAdd;
-			CombinedFileList.Append(Description.FilesForSubmit);
-
-			if(Description.FilesForAdd.Num() > 0)
-			{
-				bCheckInSuccess &= (SourceControlProvider.Execute(ISourceControlOperation::Create<FMarkForAdd>(), Description.FilesForAdd) == ECommandResult::Succeeded);
-			}
-
-			if(CombinedFileList.Num() > 0)
-			{
-				TSharedRef<FCheckIn, ESPMode::ThreadSafe> CheckInOperation = ISourceControlOperation::Create<FCheckIn>();
-				CheckInOperation->SetDescription(Description.Description);
-				bCheckInSuccess &= (SourceControlProvider.Execute(CheckInOperation, CombinedFileList) == ECommandResult::Succeeded);
-
-				if(bCheckInSuccess)
-				{
-					// report success with a notification
-					FNotificationInfo Info(CheckInOperation->GetSuccessMessage());
-					Info.ExpireDuration = 8.0f;
-					Info.HyperlinkText = LOCTEXT("SCC_Checkin_ShowLog", "Show Message Log");
-					Info.Hyperlink = FSimpleDelegate::CreateStatic([](){ FMessageLog("SourceControl").Open(EMessageSeverity::Info, true); });
-					FSlateNotificationManager::Get().AddNotification(Info);
-
-					// also add to the log
-					FMessageLog("SourceControl").Info(CheckInOperation->GetSuccessMessage());
-				}
-			}
-
-			if(!bCheckInSuccess)
-			{
-				FMessageLog("SourceControl").Notify(LOCTEXT("SCC_Checkin_Failed", "Failed to check in files!"));
-			}
-			// If we checked in ok, do we want to to re-check out the files we just checked in?
-			else
-			{
-				if( SourceControlWidget->WantToKeepCheckedOut() == true )
-				{
-					// Re-check out files
-					if(SourceControlProvider.Execute(ISourceControlOperation::Create<FCheckOut>(), CombinedFileList) != ECommandResult::Succeeded)
-					{
-						FMessageDialog::Open( EAppMsgType::Ok, LOCTEXT("SCC_Checkin_ReCheckOutFailed", "Failed to re-check out files.") );
-					}
-				}
-			}
-		}
-	}
-	else
-	{
-		FMessageLog EditorErrors("EditorErrors");
-		EditorErrors.Warning(LOCTEXT("NoAssetsToCheckIn", "No assets to check in!"));
-		EditorErrors.Notify();
-	}
-
-	return bCheckInSuccess;
-}
-
-
 #undef LOCTEXT_NAMESPACE
 
 #endif // SOURCE_CONTROL_WITH_SLATE
diff --git a/Engine/Source/Editor/SourceControlWindows/Private/SSourceControlSubmit.h b/Engine/Source/Editor/SourceControlWindows/Private/SSourceControlSubmit.h
--- a/Engine/Source/Editor/SourceControlWindows/Private/SSourceControlSubmit.h
+++ b/Engine/Source/Editor/SourceControlWindows/Private/SSourceControlSubmit.h
@@ -8,6 +8,7 @@
 #include "Widgets/DeclarativeSyntaxSupport.h"
 #include "Widgets/SWidget.h"
 #include "Widgets/SCompoundWidget.h"
+#include "Widgets/Views/SListView.h"
 #include "Widgets/Views/STableViewBase.h"
 #include "Widgets/Views/STableRow.h"
 #include "ISourceControlState.h"
diff --git a/Engine/Source/Editor/SourceControlWindows/Public/SourceControlWindows.h b/Engine/Source/Editor/SourceControlWindows/Public/SourceControlWindows.h
--- a/Engine/Source/Editor/SourceControlWindows/Public/SourceControlWindows.h
+++ b/Engine/Source/Editor/SourceControlWindows/Public/SourceControlWindows.h
@@ -6,15 +6,59 @@
 #include "SlateFwd.h"
 #include "ISourceControlProvider.h"
 
+
+/** Info supplied as argument to delegate FSourceControlWindowsOnCheckInComplete called by FSourceControlWindows::ChoosePackagesToCheckIn() and optional argument to PromptForCheckin(). */
+struct SOURCECONTROLWINDOWS_API FCheckinResultInfo
+{
+	/** Default Constructor */
+	FCheckinResultInfo();
+
+	/** Succeeded - if packages were selected and successfully checked in, Cancelled - if the user aborted the process, Failed - if an issue was encountered during the process */
+	ECommandResult::Type Result;
+
+	/** true if added and modified files were automatically checked out from source control again after being submitted, false if not */
+	bool bAutoCheckedOut;
+
+	/** Files that were added. */
+	TArray<FString> FilesAdded;
+
+	/** Files that were modified and checked in. */
+	TArray<FString> FilesSubmitted;
+
+	/** Text that describes result whether failed, cancelled or successful. */
+	FText Description;
+};
+
+/** Optional delegate called when FSourceControlWindows::ChoosePackagesToCheckIn() completes. */
+DECLARE_DELEGATE_OneParam(FSourceControlWindowsOnCheckInComplete, const FCheckinResultInfo&);
+
 class SOURCECONTROLWINDOWS_API FSourceControlWindows
 {
 public:
-	/** Opens a dialog to choose packages to submit */
-	static void ChoosePackagesToCheckIn();
+	/**
+	 * Opens a user dialog to choose packages to submit.
+	 *
+	 * @param	OnCompleteDelegate	Delegate to call when this user-based operation is complete. Also see FCheckinResultInfo.
+	 * @return	true - if command successfully in progress and OnCompleteDelegate will be called when complete, false - if immediately unable to comply (such as source control not enabled)
+	 */
+	static bool ChoosePackagesToCheckIn(const FSourceControlWindowsOnCheckInComplete& OnCompleteDelegate = FSourceControlWindowsOnCheckInComplete());
 
 	/** Determines whether we can choose packages to check in (we cant if an operation is already in progress) */
 	static bool CanChoosePackagesToCheckIn();
 
+	/**
+	 * Display check in dialog for the specified packages and get additional result information
+	 *
+	 * @param	OutResultInfo					Optional FCheckinResultInfo to store information about checked in files and whether the check in was successful
+	 * @param	InPackageNames					Names of packages to check in
+	 * @param	InPendingDeletePaths			Directories to check for files marked 'pending delete'
+	 * @param	InConfigFiles					Config filenames to check in
+	 * @param	bUseSourceControlStateCache		Whether to use the cached source control status, or force the status to be updated
+	 * 
+	 * @return	true - if completed successfully in progress and OnCompleteDelegate will be called when complete, false - if immediately unable to comply (such as source control not enabled)
+	 */
+	static bool PromptForCheckin(FCheckinResultInfo& OutResultInfo, const TArray<FString>& InPackageNames, const TArray<FString>& InPendingDeletePaths = TArray<FString>(), const TArray<FString>& InConfigFiles = TArray<FString>(), bool bUseSourceControlStateCache = false);
+
 	/**
 	 * Display check in dialog for the specified packages
 	 *
@@ -22,6 +66,8 @@ public:
 	 * @param	InPackageNames					Names of packages to check in
 	 * @param	InPendingDeletePaths			Directories to check for files marked 'pending delete'
 	 * @param	InConfigFiles					Config filenames to check in
+	 * 
+	 * @return	true - if completed successfully in progress and OnCompleteDelegate will be called when complete, false - if immediately unable to comply (such as source control not enabled)
 	 */
 	static bool PromptForCheckin(bool bUseSourceControlStateCache, const TArray<FString>& InPackageNames, const TArray<FString>& InPendingDeletePaths = TArray<FString>(), const TArray<FString>& InConfigFiles = TArray<FString>());
 
@@ -43,10 +89,10 @@ public:
 
 protected:
 	/** Callback for ChoosePackagesToCheckIn(), continues to bring up UI once source control operations are complete */
-	static void ChoosePackagesToCheckInCallback(const FSourceControlOperationRef& InOperation, ECommandResult::Type InResult);
+	static void ChoosePackagesToCheckInCallback(const FSourceControlOperationRef& InOperation, ECommandResult::Type InResult, FSourceControlWindowsOnCheckInComplete OnCompleteDelegate);
 
-	/** Called when the process has completed and we have packages to check in */
-	static void ChoosePackagesToCheckInCompleted(const TArray<UPackage*>& LoadedPackages, const TArray<FString>& PackageNames, const TArray<FString>& ConfigFiles);
+	/** Called when the user selection process has completed and we have packages to check in. */
+	static void ChoosePackagesToCheckInCompleted(const TArray<UPackage*>& LoadedPackages, const TArray<FString>& PackageNames, const TArray<FString>& ConfigFiles, FCheckinResultInfo& OutResultInfo);
 
 	/** Delegate called when the user has decided to cancel the check in process */
 	static void ChoosePackagesToCheckInCancelled(FSourceControlOperationRef InOperation);
diff --git a/Engine/Source/Editor/StaticMeshEditor/Private/StaticMeshEditorTools.cpp b/Engine/Source/Editor/StaticMeshEditor/Private/StaticMeshEditorTools.cpp
--- a/Engine/Source/Editor/StaticMeshEditor/Private/StaticMeshEditorTools.cpp
+++ b/Engine/Source/Editor/StaticMeshEditor/Private/StaticMeshEditorTools.cpp
@@ -2228,33 +2228,8 @@ void FMeshMaterialsLayout::GetMaterials(IMaterialListBuilder& ListBuilder)
 void FMeshMaterialsLayout::OnMaterialChanged(UMaterialInterface* NewMaterial, UMaterialInterface* PrevMaterial, int32 MaterialIndex, bool bReplaceAll)
 {
 	UStaticMesh& StaticMesh = GetStaticMesh();
-	FScopedTransaction ScopeTransaction(LOCTEXT("StaticMeshEditorMaterialChanged", "Staticmesh editor: Material changed"));
-
-	// flag the property (Materials) we're modifying so that not all of the object is rebuilt.
-	UProperty* ChangedProperty = NULL;
-	ChangedProperty = FindField<UProperty>(UStaticMesh::StaticClass(), "StaticMaterials");
-	check(ChangedProperty);
-	StaticMesh.PreEditChange(ChangedProperty);
-
-	if (StaticMesh.StaticMaterials.IsValidIndex(MaterialIndex))
-	{
-		StaticMesh.StaticMaterials[MaterialIndex].MaterialInterface = NewMaterial;
-		if (NewMaterial != nullptr)
-		{
-			//Set the Material slot name to a good default one
-			if (StaticMesh.StaticMaterials[MaterialIndex].MaterialSlotName == NAME_None)
-			{
-				StaticMesh.StaticMaterials[MaterialIndex].MaterialSlotName = NewMaterial->GetFName();
-			}
-			//Set the original fbx material name so we can re-import correctly
-			if (StaticMesh.StaticMaterials[MaterialIndex].ImportedMaterialSlotName == NAME_None)
-			{
-				StaticMesh.StaticMaterials[MaterialIndex].ImportedMaterialSlotName = NewMaterial->GetFName();
-			}
-		}
-	}
-
-	CallPostEditChange(ChangedProperty);
+	StaticMesh.SetMaterial(MaterialIndex, NewMaterial);
+	StaticMeshEditor.RefreshViewport();
 }
 
 TSharedRef<SWidget> FMeshMaterialsLayout::OnGenerateWidgetsForMaterial(UMaterialInterface* Material, int32 SlotIndex)
diff --git a/Engine/Source/Editor/TranslationEditor/Private/TranslationDataManager.cpp b/Engine/Source/Editor/TranslationEditor/Private/TranslationDataManager.cpp
--- a/Engine/Source/Editor/TranslationEditor/Private/TranslationDataManager.cpp
+++ b/Engine/Source/Editor/TranslationEditor/Private/TranslationDataManager.cpp
@@ -10,11 +10,11 @@
 #include "Dom/JsonObject.h"
 #include "Serialization/JsonSerializer.h"
 #include "EditorStyleSet.h"
-#include "ISourceControlOperation.h"
 #include "SourceControlOperations.h"
 #include "ISourceControlState.h"
 #include "ISourceControlProvider.h"
 #include "ISourceControlModule.h"
+#include "SourceControlHelpers.h"
 #include "TranslationUnit.h"
 #include "Logging/MessageLog.h"
 #include "TextLocalizationResourceGenerator.h"
@@ -329,7 +329,7 @@ bool FTranslationDataManager::WriteJSONToTextFile(TSharedRef<FJsonObject>& Outpu
 			{
 				bPreviouslyCheckedOut = true;
 			}
-			else if (!SourceControlHelpers::CheckOutFile(Filename))
+			else if (!USourceControlHelpers::CheckOutOrAddFile(Filename))
 			{
 				FFormatNamedArguments Arguments;
 				Arguments.Add(TEXT("Filename"), FText::FromString(Filename));
@@ -360,7 +360,7 @@ bool FTranslationDataManager::WriteJSONToTextFile(TSharedRef<FJsonObject>& Outpu
 				{
 					bPreviouslyCheckedOut = false;
 
-					if( !SourceControlHelpers::CheckOutFile(Filename) )
+					if( !USourceControlHelpers::CheckOutOrAddFile(Filename) )
 					{
 						FFormatNamedArguments Arguments;
 						Arguments.Add( TEXT("Filename"), FText::FromString(Filename) );
