Fix an uninitialized variable in SourceControlOperations.h

diff --git a/Engine/Source/Developer/SourceControl/Public/SourceControlOperations.h b/Engine/Source/Developer/SourceControl/Public/SourceControlOperations.h
--- a/Engine/Source/Developer/SourceControl/Public/SourceControlOperations.h
+++ b/Engine/Source/Developer/SourceControl/Public/SourceControlOperations.h
@@ -206,6 +206,7 @@ public:
 		: bUpdateHistory(false)
 		, bGetOpenedOnly(false)
 		, bUpdateModifiedState(false)
+		, bCheckingAllFiles(false)
 	{
 	}
 
