Few extra missing polish items for source control dialog

diff --git a/Engine/Source/Developer/SourceControl/Private/SSourceControlLogin.cpp b/Engine/Source/Developer/SourceControl/Private/SSourceControlLogin.cpp
--- a/Engine/Source/Developer/SourceControl/Private/SSourceControlLogin.cpp
+++ b/Engine/Source/Developer/SourceControl/Private/SSourceControlLogin.cpp
@@ -119,8 +119,8 @@ void SSourceControlLogin::Construct(const FArguments& InArgs)
 				]
 	#endif
 				+SVerticalBox::Slot()
-				.AutoHeight()
-				.Padding(8.0f, 16.0f, 8.0f, 8.0f)
+				.VAlign(VAlign_Bottom)
+				.Padding(8.0f, 16.0f, 8.0f, 16.0f)
 				[
 					SNew(SHorizontalBox)
 					+SHorizontalBox::Slot()
@@ -142,6 +142,7 @@ void SSourceControlLogin::Construct(const FArguments& InArgs)
 						[
 							SNew(SButton)
 							.VAlign(VAlign_Center)
+							.ButtonStyle(FAppStyle::Get(), "PrimaryButton")
 							.Text(LOCTEXT("AcceptSettings", "Accept Settings"))
 							.OnClicked( this, &SSourceControlLogin::OnAcceptSettings )
 							.IsEnabled( this, &SSourceControlLogin::IsAcceptSettingsEnabled )
