Source control window speedup - USourceControlHelpers::GetAssetData() should only consider assets on disk

diff --git a/Engine/Source/Developer/SourceControl/Private/SourceControlHelpers.cpp b/Engine/Source/Developer/SourceControl/Private/SourceControlHelpers.cpp
--- a/Engine/Source/Developer/SourceControl/Private/SourceControlHelpers.cpp
+++ b/Engine/Source/Developer/SourceControl/Private/SourceControlHelpers.cpp
@@ -1452,7 +1452,7 @@ bool USourceControlHelpers::GetAssetData(const FString & InFileName, const FStri
 
 	// Try the registry first
 	FAssetRegistryModule& AssetRegistryModule = FModuleManager::LoadModuleChecked<FAssetRegistryModule>(TEXT("AssetRegistry"));
-	AssetRegistryModule.Get().GetAssetsByPackageName(*InPackageName, OutAssets);
+	AssetRegistryModule.Get().GetAssetsByPackageName(*InPackageName, OutAssets, true);
 
 	if (OutAssets.Num() > 0)
 	{
