WP HLODBuilder - Fix error when there's no files to mark for add
* P4 SCC would give an "invalid syntax" error as the p4 command was ill formed
* Added early exits in USourceControlHelpers operations that take an array of files

diff --git a/Engine/Plugins/Developer/PerforceSourceControl/Source/PerforceSourceControl/Private/PerforceSourceControlOperations.cpp b/Engine/Plugins/Developer/PerforceSourceControl/Source/PerforceSourceControl/Private/PerforceSourceControlOperations.cpp
--- a/Engine/Plugins/Developer/PerforceSourceControl/Source/PerforceSourceControl/Private/PerforceSourceControlOperations.cpp
+++ b/Engine/Plugins/Developer/PerforceSourceControl/Source/PerforceSourceControl/Private/PerforceSourceControlOperations.cpp
@@ -542,6 +542,12 @@ FName FPerforceMarkForAddWorker::GetName() const
 
 bool FPerforceMarkForAddWorker::Execute(FPerforceSourceControlCommand& InCommand)
 {
+	// Avoid invalid p4 syntax if there's no file to process
+	if (InCommand.Files.IsEmpty())
+	{
+		return true;
+	}
+
 	// Perforce will allow you to mark files for add that don't currently exist on disk
 	// This goes against the workflow of our other SCC providers (such as SVN and Git),
 	// so we manually check that the files exist before allowing this command to continue
diff --git a/Engine/Source/Developer/SourceControl/Private/SourceControlHelpers.cpp b/Engine/Source/Developer/SourceControl/Private/SourceControlHelpers.cpp
--- a/Engine/Source/Developer/SourceControl/Private/SourceControlHelpers.cpp
+++ b/Engine/Source/Developer/SourceControl/Private/SourceControlHelpers.cpp
@@ -293,6 +293,12 @@ bool USourceControlHelpers::SyncFile(const FString& InFile, bool bSilent)
 
 bool USourceControlHelpers::SyncFiles(const TArray<FString>& InFiles, bool bSilent)
 {
+	// If we have nothing to process, exit immediately 
+	if (InFiles.IsEmpty())
+	{
+		return true;
+	}
+
 	ISourceControlProvider* Provider = SourceControlHelpersInternal::VerifySourceControl(bSilent);
 
 	if (!Provider)
@@ -397,6 +403,12 @@ bool USourceControlHelpers::CheckOutFile(const FString& InFile, bool bSilent)
 
 bool USourceControlHelpers::CheckOutFiles(const TArray<FString>& InFiles, bool bSilent)
 {
+	// If we have nothing to process, exit immediately
+	if (InFiles.IsEmpty())
+	{
+		return true;
+	}
+
 	// Determine file type and ensure it is in form source control wants
 	// Even if some files were skipped, still apply to the others
 	TArray<FString> SCFiles;
@@ -543,6 +555,12 @@ bool USourceControlHelpers::CheckOutOrAddFile(const FString& InFile, bool bSilen
 
 bool USourceControlHelpers::CheckOutOrAddFiles(const TArray<FString>& InFiles, bool bSilent)
 {
+	// If we have nothing to process, exit immediately
+	if (InFiles.IsEmpty())
+	{
+		return true;
+	}
+
 	// Determine file type and ensure it is in form source control wants
 	// Even if some files were skipped, still apply to the others
 	TArray<FString> SCFiles;
@@ -664,6 +682,12 @@ bool USourceControlHelpers::MarkFileForAdd(const FString& InFile, bool bSilent)
 
 bool USourceControlHelpers::MarkFilesForAdd(const TArray<FString>& InFiles, bool bSilent)
 {
+	// If we have nothing to process, exit immediately
+	if (InFiles.IsEmpty())
+	{
+		return true;
+	}
+
 	ISourceControlProvider* Provider = SourceControlHelpersInternal::VerifySourceControl(bSilent);
 
 	if (!Provider)
@@ -767,6 +791,12 @@ bool USourceControlHelpers::MarkFileForDelete(const FString& InFile, bool bSilen
 
 bool USourceControlHelpers::MarkFilesForDelete(const TArray<FString>& InFiles, bool bSilent)
 {
+	// If we have nothing to process, exit immediately
+	if (InFiles.IsEmpty())
+	{
+		return true;
+	}
+
 	// Determine file type and ensure it is in form source control wants
 	// Even if some files were skipped, still apply to the others
 	TArray<FString> SCFiles;
@@ -871,6 +901,12 @@ bool USourceControlHelpers::RevertFile(const FString& InFile, bool bSilent)
 
 bool USourceControlHelpers::RevertFiles(const TArray<FString>& InFiles,	bool bSilent)
 {
+	// If we have nothing to process, exit immediately
+	if (InFiles.IsEmpty())
+	{
+		return true;
+	}
+
 	// Determine file type and ensure they are in form source control wants
 	// Even if some files were skipped, still apply to the others
 	TArray<FString> SCFiles;
@@ -945,6 +981,12 @@ bool USourceControlHelpers::RevertUnchangedFile(const FString& InFile, bool bSil
 
 bool USourceControlHelpers::RevertUnchangedFiles(const TArray<FString>& InFiles, bool bSilent)
 {
+	// If we have nothing to process, exit immediately
+	if (InFiles.IsEmpty())
+	{
+		return true;
+	}
+
 	// Determine file types and ensure they are in form source control wants
 	TArray<FString> FilePaths;
 	SourceControlHelpersInternal::ConvertFilesToQualifiedPaths(InFiles, FilePaths, bSilent);
@@ -996,6 +1038,12 @@ bool USourceControlHelpers::CheckInFile(const FString& InFile, const FString& In
 
 bool USourceControlHelpers::CheckInFiles(const TArray<FString>& InFiles, const FString& InDescription, bool bSilent)
 {
+	// If we have nothing to process, exit immediately
+	if (InFiles.IsEmpty())
+	{
+		return true;
+	}
+
 	ISourceControlProvider* Provider = SourceControlHelpersInternal::VerifySourceControl(bSilent);
 
 	if (!Provider)
diff --git a/Engine/Source/Editor/UnrealEd/Private/WorldPartition/WorldPartitionHLODsBuilder.cpp b/Engine/Source/Editor/UnrealEd/Private/WorldPartition/WorldPartitionHLODsBuilder.cpp
--- a/Engine/Source/Editor/UnrealEd/Private/WorldPartition/WorldPartitionHLODsBuilder.cpp
+++ b/Engine/Source/Editor/UnrealEd/Private/WorldPartition/WorldPartitionHLODsBuilder.cpp
@@ -884,45 +884,57 @@ bool UWorldPartitionHLODsBuilder::CopyFilesFromWorkingDir(const FString& SourceD
 		}
 	}
 
-	TArray<FString> FilesToEdit;
-	FilesToCopy.GetKeys(FilesToEdit);
-
-	bRet = USourceControlHelpers::MarkFilesForAdd(FilesToAdd);
-	if (!bRet)
+	// Add
+	if (!FilesToAdd.IsEmpty())
 	{
-		UE_LOG(LogWorldPartitionHLODsBuilder, Error, TEXT("Adding files to source control failed: %s"), *USourceControlHelpers::LastErrorMsg().ToString());
-		while (!FPlatformMisc::IsDebuggerPresent())
+		bRet = USourceControlHelpers::MarkFilesForAdd(FilesToAdd);
+		if (!bRet)
 		{
-			FPlatformProcess::Sleep(0.1f);
+			UE_LOG(LogWorldPartitionHLODsBuilder, Error, TEXT("Adding files to source control failed: %s"), *USourceControlHelpers::LastErrorMsg().ToString());
+			while (!FPlatformMisc::IsDebuggerPresent())
+			{
+				FPlatformProcess::Sleep(0.1f);
+			}
+			return false;
 		}
-		return false;
+		ModifiedFiles.Append(FilesToAdd);
 	}
-	ModifiedFiles.Append(FilesToAdd);
 
-	bRet = USourceControlHelpers::MarkFilesForDelete(FilesToDelete);
-	if (!bRet)
+	// Delete
+	if (!FilesToDelete.IsEmpty())
 	{
-		UE_LOG(LogWorldPartitionHLODsBuilder, Error, TEXT("Deleting files from source control failed: %s"), *USourceControlHelpers::LastErrorMsg().ToString());
-		return false;
+		bRet = USourceControlHelpers::MarkFilesForDelete(FilesToDelete);
+		if (!bRet)
+		{
+			UE_LOG(LogWorldPartitionHLODsBuilder, Error, TEXT("Deleting files from source control failed: %s"), *USourceControlHelpers::LastErrorMsg().ToString());
+			return false;
+		}
+		ModifiedFiles.Append(FilesToDelete);
 	}
-	ModifiedFiles.Append(FilesToDelete);
 
-	bRet = USourceControlHelpers::CheckOutFiles(FilesToEdit);
-	if (!bRet)
+	// Edit
+	if (!FilesToCopy.IsEmpty())
 	{
-		UE_LOG(LogWorldPartitionHLODsBuilder, Error, TEXT("Checking out files from source control failed: %s"), *USourceControlHelpers::LastErrorMsg().ToString());
-		return false;
-	}
+		TArray<FString> FilesToEdit;
+		FilesToCopy.GetKeys(FilesToEdit);
 
-	for(const auto& Pair : FilesToCopy)
-	{
-		bRet = IFileManager::Get().Copy(*Pair.Key, *Pair.Value) == COPY_OK;
+		bRet = USourceControlHelpers::CheckOutFiles(FilesToEdit);
 		if (!bRet)
 		{
-			UE_LOG(LogWorldPartitionHLODsBuilder, Error, TEXT("Failed to copy file from \"%s\" to \"%s\""), *Pair.Value, *Pair.Key);
+			UE_LOG(LogWorldPartitionHLODsBuilder, Error, TEXT("Checking out files from source control failed: %s"), *USourceControlHelpers::LastErrorMsg().ToString());
 			return false;
 		}
-		ModifiedFiles.Add(Pair.Key);
+		ModifiedFiles.Append(FilesToEdit);
+
+		for (const auto& Pair : FilesToCopy)
+		{
+			bRet = IFileManager::Get().Copy(*Pair.Key, *Pair.Value) == COPY_OK;
+			if (!bRet)
+			{
+				UE_LOG(LogWorldPartitionHLODsBuilder, Error, TEXT("Failed to copy file from \"%s\" to \"%s\""), *Pair.Value, *Pair.Key);
+				return false;
+			}
+		}
 	}
 
 	return true;
