SourceControlHelper: Replace UE4 by UE in comment

diff --git a/Engine/Source/Developer/SourceControl/Private/SourceControlHelpers.cpp b/Engine/Source/Developer/SourceControl/Private/SourceControlHelpers.cpp
--- a/Engine/Source/Developer/SourceControl/Private/SourceControlHelpers.cpp
+++ b/Engine/Source/Developer/SourceControl/Private/SourceControlHelpers.cpp
@@ -164,7 +164,7 @@ FString ConvertFileToQualifiedPath(const FString& InFile, bool bSilent, bool bAl
 	}
 
 	// Qualify based on process base directory.
-	// Something akin to "C:/Epic/UE4/Engine/Binaries/Win64/" as a current path.
+	// Something akin to "C:/Epic/UE/Engine/Binaries/Win64/" as a current path.
 	SCFile = FPaths::ConvertRelativePathToFull(InFile);
 
 	if (FPaths::FileExists(SCFile) || (bAllowDirectories && FPaths::DirectoryExists(SCFile)))
