FScopedSourceControl: Only Init/Close Provider if it isn't Enabled yet

diff --git a/Engine/Source/Developer/SourceControl/Private/SourceControlHelpers.cpp b/Engine/Source/Developer/SourceControl/Private/SourceControlHelpers.cpp
--- a/Engine/Source/Developer/SourceControl/Private/SourceControlHelpers.cpp
+++ b/Engine/Source/Developer/SourceControl/Private/SourceControlHelpers.cpp
@@ -1597,13 +1597,20 @@ bool USourceControlHelpers::GetAssetDataFromFileHistory(FSourceControlStatePtr I
 
 FScopedSourceControl::FScopedSourceControl()
 {
-	ISourceControlModule::Get().GetProvider().Init();
+	bInitSourceControl = !ISourceControlModule::Get().GetProvider().IsEnabled();
+	if (bInitSourceControl)
+	{
+		ISourceControlModule::Get().GetProvider().Init();
+	}
 }
 
 
 FScopedSourceControl::~FScopedSourceControl()
 {
-	ISourceControlModule::Get().GetProvider().Close();
+	if (bInitSourceControl)
+	{
+		ISourceControlModule::Get().GetProvider().Close();
+	}
 }
 
 
diff --git a/Engine/Source/Developer/SourceControl/Public/SourceControlHelpers.h b/Engine/Source/Developer/SourceControl/Public/SourceControlHelpers.h
--- a/Engine/Source/Developer/SourceControl/Public/SourceControlHelpers.h
+++ b/Engine/Source/Developer/SourceControl/Public/SourceControlHelpers.h
@@ -561,4 +561,7 @@ public:
 
 	/** Get the provider we are using */
 	ISourceControlProvider& GetProvider();
+
+private:
+	bool bInitSourceControl = false;
 };
