FScopedSourceControl - Fix wrong condition preventing source control initialisation

diff --git a/Engine/Source/Developer/SourceControl/Private/SourceControlHelpers.cpp b/Engine/Source/Developer/SourceControl/Private/SourceControlHelpers.cpp
--- a/Engine/Source/Developer/SourceControl/Private/SourceControlHelpers.cpp
+++ b/Engine/Source/Developer/SourceControl/Private/SourceControlHelpers.cpp
@@ -1597,7 +1597,7 @@ bool USourceControlHelpers::GetAssetDataFromFileHistory(FSourceControlStatePtr I
 
 FScopedSourceControl::FScopedSourceControl()
 {
-	bInitSourceControl = !ISourceControlModule::Get().GetProvider().IsEnabled();
+	bInitSourceControl = !ISourceControlModule::Get().GetProvider().IsAvailable();
 	if (bInitSourceControl)
 	{
 		ISourceControlModule::Get().GetProvider().Init();
