ISourceControlOperation.h::Create can now forward arguments to the constructor of the operation being created.

- This will avoid the need to create the operation and the immediatly call a number of setters to construct it properly.

diff --git a/Engine/Source/Developer/SourceControl/Public/ISourceControlOperation.h b/Engine/Source/Developer/SourceControl/Public/ISourceControlOperation.h
--- a/Engine/Source/Developer/SourceControl/Public/ISourceControlOperation.h
+++ b/Engine/Source/Developer/SourceControl/Public/ISourceControlOperation.h
@@ -70,12 +70,11 @@ public:
 	}
 
 	/** Factory method for easier operation creation */
-	template<typename Type>
-	static TSharedRef<Type, ESPMode::ThreadSafe> Create()
+	template<typename Type, typename... TArgs>
+	static TSharedRef<Type, ESPMode::ThreadSafe> Create(TArgs&&... Args)
 	{
-		return MakeShareable( new Type() );
+		return MakeShareable( new Type(Forward<TArgs>(Args)...));
 	}
-
 };
 
 typedef TSharedRef<class ISourceControlOperation, ESPMode::ThreadSafe> FSourceControlOperationRef;
