Fix apparent typo in Source Control dialog caused by superfluous padding.

diff --git a/Engine/Source/Developer/SourceControl/Private/SSourceControlLogin.cpp b/Engine/Source/Developer/SourceControl/Private/SSourceControlLogin.cpp
--- a/Engine/Source/Developer/SourceControl/Private/SSourceControlLogin.cpp
+++ b/Engine/Source/Developer/SourceControl/Private/SSourceControlLogin.cpp
@@ -79,7 +79,6 @@ void SSourceControlLogin::Construct(const FArguments& InArgs)
 					]
 					+SVerticalBox::Slot()
 					.AutoHeight()
-					.Padding(FMargin(0.0f, 0.0f, 16.0f, 0.0f))
 					[
 						SNew(STextBlock)
 						.Visibility(this, &SSourceControlLogin::GetDisabledTextVisibility)
