Revert source control button changes until I can figure out why it doesnt link on linux.

diff --git a/Engine/Source/Developer/SourceControl/Private/SSourceControlLogin.cpp b/Engine/Source/Developer/SourceControl/Private/SSourceControlLogin.cpp
--- a/Engine/Source/Developer/SourceControl/Private/SSourceControlLogin.cpp
+++ b/Engine/Source/Developer/SourceControl/Private/SSourceControlLogin.cpp
@@ -15,7 +15,6 @@
 #include "Widgets/Layout/SExpandableArea.h"
 #include "Framework/Notifications/NotificationManager.h"
 #include "Widgets/Notifications/SNotificationList.h"
-#include "SPrimaryButton.h"
 #endif
 
 #include "SSourceControlPicker.h"
@@ -140,7 +139,10 @@ void SSourceControlLogin::Construct(const FArguments& InArgs)
 						.AutoWidth()
 						.Padding(FMargin(5.0f, 0.0f))
 						[
-							SNew(SPrimaryButton)
+							SNew(SButton)
+							.VAlign(VAlign_Center)
+							.ButtonStyle(FAppStyle::Get(), "PrimaryButton")
+  							.TextStyle( FAppStyle::Get(), "DialogButtonText" )
 							.Text(LOCTEXT("AcceptSettings", "Accept Settings"))
 							.OnClicked( this, &SSourceControlLogin::OnAcceptSettings )
 							.IsEnabled( this, &SSourceControlLogin::IsAcceptSettingsEnabled )
@@ -150,6 +152,8 @@ void SSourceControlLogin::Construct(const FArguments& InArgs)
 						.Padding(FMargin(5.0f, 0.0f))
 						[
 							SNew(SButton)
+							.VAlign(VAlign_Center)
+  							.TextStyle( FAppStyle::Get(), "DialogButtonText" )
 							.Text(LOCTEXT("RunWithoutSourceControl", "Disable Source Control"))
 							.OnClicked( this, &SSourceControlLogin::OnDisableSourceControl )
 							.IsEnabled( this, &SSourceControlLogin::AreControlsEnabled )
diff --git a/Engine/Source/Developer/SourceControl/SourceControl.Build.cs b/Engine/Source/Developer/SourceControl/SourceControl.Build.cs
--- a/Engine/Source/Developer/SourceControl/SourceControl.Build.cs
+++ b/Engine/Source/Developer/SourceControl/SourceControl.Build.cs
@@ -24,7 +24,6 @@ public class SourceControl : ModuleRules
 				"AssetRegistry",
 				"Slate",
 				"EditorStyle",
-				"ToolWidgets",
 			}
 		);
 
