Failure to checkout files during the HLOD generation process gives no context at all

diff --git a/Engine/Source/Developer/SourceControl/Private/SourceControlHelpers.cpp b/Engine/Source/Developer/SourceControl/Private/SourceControlHelpers.cpp
--- a/Engine/Source/Developer/SourceControl/Private/SourceControlHelpers.cpp
+++ b/Engine/Source/Developer/SourceControl/Private/SourceControlHelpers.cpp
@@ -318,6 +318,36 @@ bool USourceControlHelpers::SyncFiles(const TArray<FString>& InFiles, bool bSile
 	return !bFilesSkipped && (Result == ECommandResult::Succeeded);
 }
 
+void LogCheckoutFailure(const FString& InFile, const FString& SCFile, FSourceControlStatePtr SCState, bool bCheckoutFailed, bool bSilent)
+{
+	FString SimultaneousCheckoutUser;
+	FFormatNamedArguments Arguments;
+	Arguments.Add(TEXT("InFile"), FText::FromString(InFile));
+	Arguments.Add(TEXT("SCFile"), FText::FromString(SCFile));
+
+	if (bCheckoutFailed)
+	{
+		SourceControlHelpersInternal::LogError(FText::Format(LOCTEXT("CheckoutFailed", "Failed to check out file '{InFile}' ({SCFile})."), Arguments), bSilent);
+	}
+	else if (!SCState->IsSourceControlled())
+	{
+		SourceControlHelpersInternal::LogError(FText::Format(LOCTEXT("NotSourceControlled", "Could not check out the file '{InFile}' because it is not under source control ({SCFile})."), Arguments), bSilent);
+	}
+	else if (!SCState->IsCurrent())
+	{
+		SourceControlHelpersInternal::LogError(FText::Format(LOCTEXT("NotAtHeadRevision", "File '{InFile}' is not at head revision ({SCFile})."), Arguments), bSilent);
+	}
+	else if (SCState->IsCheckedOutOther(&(SimultaneousCheckoutUser)))
+	{
+		Arguments.Add(TEXT("SimultaneousCheckoutUser"), FText::FromString(SimultaneousCheckoutUser));
+		SourceControlHelpersInternal::LogError(FText::Format(LOCTEXT("SimultaneousCheckout", "File '{InFile}' is checked out by another ({SimultaneousCheckoutUser}) ({SCFile})."), Arguments), bSilent);
+	}
+	else
+	{
+		// Improper or invalid SCC state
+		SourceControlHelpersInternal::LogError(FText::Format(LOCTEXT("CouldNotDetermineState", "Could not determine source control state of file '{InFile}' ({SCFile})."), Arguments), bSilent);
+	}
+}
 
 bool USourceControlHelpers::CheckOutFile(const FString& InFile, bool bSilent)
 {
@@ -341,12 +371,7 @@ bool USourceControlHelpers::CheckOutFile(const FString& InFile, bool bSilent)
 
 	if (!SCState.IsValid())
 	{
-		// Improper or invalid SCC state
-		FFormatNamedArguments Arguments;
-		Arguments.Add(TEXT("InFile"), FText::FromString(InFile));
-		Arguments.Add(TEXT("SCFile"), FText::FromString(SCFile));
-		SourceControlHelpersInternal::LogError(FText::Format(LOCTEXT("CouldNotDetermineState", "Could not determine source control state of file '{InFile}' ({SCFile})."), Arguments), bSilent);
-
+		LogCheckoutFailure(InFile, SCFile, SCState, false, bSilent);
 		return false;
 	}
 
@@ -356,7 +381,7 @@ bool USourceControlHelpers::CheckOutFile(const FString& InFile, bool bSilent)
 		return true;
 	}
 
-	bool bCheckOutFail = false;
+	bool bCheckOutFailed = false;
 
 	if (SCState->CanCheckout())
 	{
@@ -365,38 +390,10 @@ bool USourceControlHelpers::CheckOutFile(const FString& InFile, bool bSilent)
 			return true;
 		}
 
-		bCheckOutFail = true;
+		bCheckOutFailed = true;
 	}
 
-	// Only error info after this point
-
-	FString SimultaneousCheckoutUser;
-	FFormatNamedArguments Arguments;
-	Arguments.Add(TEXT("InFile"), FText::FromString(InFile));
-	Arguments.Add(TEXT("SCFile"), FText::FromString(SCFile));
-
-	if (bCheckOutFail)
-	{
-		SourceControlHelpersInternal::LogError(FText::Format(LOCTEXT("CheckoutFailed", "Failed to check out file '{InFile}' ({SCFile})."), Arguments), bSilent);
-	}
-	else if (!SCState->IsSourceControlled())
-	{
-		SourceControlHelpersInternal::LogError(FText::Format(LOCTEXT("NotSourceControlled", "Could not check out the file '{InFile}' because it is not under source control ({SCFile})."), Arguments), bSilent);
-	}
-	else if (!SCState->IsCurrent())
-	{
-		SourceControlHelpersInternal::LogError(FText::Format(LOCTEXT("NotAtHeadRevision", "File '{InFile}' is not at head revision ({SCFile})."), Arguments), bSilent);
-	}
-	else if (SCState->IsCheckedOutOther(&(SimultaneousCheckoutUser)))
-	{
-		Arguments.Add(TEXT("SimultaneousCheckoutUser"), FText::FromString(SimultaneousCheckoutUser));
-		SourceControlHelpersInternal::LogError(FText::Format(LOCTEXT("SimultaneousCheckout", "File '{InFile}' is checked out by another ({SimultaneousCheckoutUser}) ({SCFile})."), Arguments), bSilent);
-	}
-	else
-	{
-		// Improper or invalid SCC state
-		SourceControlHelpersInternal::LogError(FText::Format(LOCTEXT("CouldNotDetermineState", "Could not determine source control state of file '{InFile}' ({SCFile})."), Arguments), bSilent);
-	}
+	LogCheckoutFailure(InFile, SCFile, SCState, bCheckOutFailed, bSilent);
 
 	return false;
 }
@@ -433,8 +430,8 @@ bool USourceControlHelpers::CheckOutFiles(const TArray<FString>& InFiles, bool b
 		FString SCFile = SCFiles[Index];
 		FSourceControlStateRef SCState = SCStates[Index];
 
-	// Less error checking and info is made for multiple files than the single file version.
-	// This multi-file version could be made similarly more sophisticated.
+		// Less error checking and info is made for multiple files than the single file version.
+		// This multi-file version could be made similarly more sophisticated.
 		if (!SCState->IsCheckedOut() && !SCState->IsAdded())
 		{
 			if (SCState->CanCheckout())
@@ -444,14 +441,15 @@ bool USourceControlHelpers::CheckOutFiles(const TArray<FString>& InFiles, bool b
 			else
 			{
 				bCannotCheckoutAtLeastOneFile = true;
+				LogCheckoutFailure(InFiles[Index], SCFile, SCState, false, bSilent);
 			}
 		}
 	}
 
 	bool bSuccess = !bFilesSkipped && !bCannotCheckoutAtLeastOneFile;
-	if (SCFilesToCheckout.Num())
+	if (bSuccess && SCFilesToCheckout.Num())
 	{
-		bSuccess &= Provider->Execute(ISourceControlOperation::Create<FCheckOut>(), SCFilesToCheckout) == ECommandResult::Succeeded;
+		bSuccess = Provider->Execute(ISourceControlOperation::Create<FCheckOut>(), SCFilesToCheckout) == ECommandResult::Succeeded;
 	}
 
 	return bSuccess;
