Adopting some editor source control modifications from a partner stream.
  * Introduces a new 'GetFileList' operation
  * Adds "keep checked out" option to other operations

diff --git a/Engine/Plugins/Developer/PerforceSourceControl/Source/PerforceSourceControl/Private/PerforceSourceControlModule.cpp b/Engine/Plugins/Developer/PerforceSourceControl/Source/PerforceSourceControl/Private/PerforceSourceControlModule.cpp
--- a/Engine/Plugins/Developer/PerforceSourceControl/Source/PerforceSourceControl/Private/PerforceSourceControlModule.cpp
+++ b/Engine/Plugins/Developer/PerforceSourceControl/Source/PerforceSourceControl/Private/PerforceSourceControlModule.cpp
@@ -40,6 +40,7 @@ void FPerforceSourceControlModule::StartupModule()
 	PerforceSourceControlProvider.RegisterWorker( "DownloadFile", FGetPerforceSourceControlWorker::CreateStatic( &CreateWorker<FPerforceDownloadFileWorker>) );
 	PerforceSourceControlProvider.RegisterWorker( "CreateWorkspace", FGetPerforceSourceControlWorker::CreateStatic( &CreateWorker<FPerforceCreateWorkspaceWorker>) );
 	PerforceSourceControlProvider.RegisterWorker( "DeleteWorkspace", FGetPerforceSourceControlWorker::CreateStatic( &CreateWorker<FPerforceDeleteWorkspaceWorker>) );
+	PerforceSourceControlProvider.RegisterWorker( "GetFileList", FGetPerforceSourceControlWorker::CreateStatic(&CreateWorker<FPerforceGetFileListWorker>));
 
 	// load our settings
 	PerforceSourceControlSettings.LoadSettings();
diff --git a/Engine/Plugins/Developer/PerforceSourceControl/Source/PerforceSourceControl/Private/PerforceSourceControlOperations.cpp b/Engine/Plugins/Developer/PerforceSourceControl/Source/PerforceSourceControl/Private/PerforceSourceControlOperations.cpp
--- a/Engine/Plugins/Developer/PerforceSourceControl/Source/PerforceSourceControl/Private/PerforceSourceControlOperations.cpp
+++ b/Engine/Plugins/Developer/PerforceSourceControl/Source/PerforceSourceControl/Private/PerforceSourceControlOperations.cpp
@@ -507,6 +507,11 @@ bool FPerforceCheckInWorker::Execute(FPerforceSourceControlCommand& InCommand)
 			TArray<FString> SubmitParams;
 			FP4RecordSet Records;
 
+			if (Operation->GetKeepCheckedOut())
+			{
+				SubmitParams.Add(TEXT("-r"));
+			}
+
 			SubmitParams.Add(TEXT("-c"));
 			SubmitParams.Add(ChangeList.ToString());
 
@@ -594,6 +599,55 @@ bool FPerforceCheckInWorker::UpdateStates() const
 	return (bUpdatedStates || bUpdatedChangelistStates);
 }
 
+FName FPerforceGetFileListWorker::GetName() const
+{
+	return "GetFileList";
+}
+
+static void AppendMaskParameter(TArray<FString>& InOutParams)
+{
+	for (int i = 0; i < InOutParams.Num(); ++i)
+	{
+		InOutParams[i] += TEXT("/...");
+	}
+}
+
+bool FPerforceGetFileListWorker::Execute(FPerforceSourceControlCommand& InCommand)
+{
+	FScopedPerforceConnection ScopedConnection(InCommand);
+	if (!InCommand.IsCanceled() && ScopedConnection.IsValid())
+	{
+		FPerforceConnection& Connection = ScopedConnection.GetConnection();
+		TSharedRef<FGetFileList, ESPMode::ThreadSafe> Operation = StaticCastSharedRef<FGetFileList>(InCommand.Operation);
+
+		TArray<FString> Parameters;
+		if (!Operation->GetIncludeDeleted())
+		{
+			Parameters.Add(TEXT("-e"));
+		}
+		Parameters.Append(InCommand.Files);
+		AppendMaskParameter(Parameters);
+		FP4RecordSet Records;
+		InCommand.bCommandSuccessful = Connection.RunCommand(TEXT("files"), Parameters, Records, InCommand.ResultInfo.ErrorMessages, FOnIsCancelled::CreateRaw(&InCommand, &FPerforceSourceControlCommand::IsCanceled), InCommand.bConnectionDropped);
+		ParseRecordSetForState(Records, OutResults);
+
+		TArray<FString> FilesList;
+		FilesList.Reserve(OutResults.Num());
+		for (const auto& Pair : OutResults)
+		{
+			FilesList.Add(Pair.Key);
+		}
+		Operation->SetFilesList(MoveTemp(FilesList));
+	}
+
+	return InCommand.bCommandSuccessful;
+}
+
+bool FPerforceGetFileListWorker::UpdateStates() const
+{
+	return false;
+}
+
 FName FPerforceMarkForAddWorker::GetName() const
 {
 	return "MarkForAdd";
@@ -2892,8 +2946,6 @@ bool FPerforceCreateWorkspaceWorker::Execute(class FPerforceSourceControlCommand
 			ClientDesc << Mapping.Value << TEXT("\n");
 		}
 
-		AddType(*Operation, ClientDesc);
-		
 		InCommand.bCommandSuccessful = Connection.CreateWorkspace(	ClientDesc,
 																	FOnIsCancelled::CreateRaw(&InCommand, &FPerforceSourceControlCommand::IsCanceled),
 																	InCommand.ResultInfo.ErrorMessages);
@@ -2912,28 +2964,6 @@ FName FPerforceDeleteWorkspaceWorker::GetName() const
 	return "DeleteWorkspace";
 }
 
-void FPerforceCreateWorkspaceWorker::AddType(const FCreateWorkspace& Operation, FStringBuilderBase& ClientDesc)
-{
-	ClientDesc << TEXT("Type:\t");
-
-	switch (Operation.GetType())
-	{
-	case FCreateWorkspace::EType::Writeable:
-		ClientDesc << TEXT("writeable");
-		break;
-	case FCreateWorkspace::EType::ReadOnly:
-		ClientDesc << TEXT("readonly");
-		break;
-	case FCreateWorkspace::EType::Partitioned:
-		ClientDesc << TEXT("partitioned");
-		break;
-	default:
-		checkNoEntry();
-	}
-
-	ClientDesc << TEXT("\n");
-}
-
 bool FPerforceDeleteWorkspaceWorker::Execute(class FPerforceSourceControlCommand& InCommand)
 {
 	FScopedPerforceConnection ScopedConnection(InCommand);
diff --git a/Engine/Plugins/Developer/PerforceSourceControl/Source/PerforceSourceControl/Private/PerforceSourceControlOperations.h b/Engine/Plugins/Developer/PerforceSourceControl/Source/PerforceSourceControl/Private/PerforceSourceControlOperations.h
--- a/Engine/Plugins/Developer/PerforceSourceControl/Source/PerforceSourceControl/Private/PerforceSourceControlOperations.h
+++ b/Engine/Plugins/Developer/PerforceSourceControl/Source/PerforceSourceControl/Private/PerforceSourceControlOperations.h
@@ -55,6 +55,20 @@ public:
 	FPerforceSourceControlChangelist OutChangelist;
 };
 
+class FPerforceGetFileListWorker : public IPerforceSourceControlWorker
+{
+public:
+	virtual ~FPerforceGetFileListWorker() {}
+	// IPerforceSourceControlWorker interface
+	virtual FName GetName() const override;
+	virtual bool Execute(class FPerforceSourceControlCommand& InCommand) override;
+	virtual bool UpdateStates() const override;
+
+public:
+	/** Map of filenames to perforce state */
+	TMap<FString, EPerforceState::Type> OutResults;
+};
+
 class FPerforceMarkForAddWorker : public IPerforceSourceControlWorker
 {
 public:
diff --git a/Engine/Source/Developer/SourceControl/Private/SourceControlHelpers.cpp b/Engine/Source/Developer/SourceControl/Private/SourceControlHelpers.cpp
--- a/Engine/Source/Developer/SourceControl/Private/SourceControlHelpers.cpp
+++ b/Engine/Source/Developer/SourceControl/Private/SourceControlHelpers.cpp
@@ -111,7 +111,7 @@ FString ConvertFileToQualifiedPath(const FString& InFile, bool bSilent, bool bAl
 	}
 
 	// Package paths
-	if (SCFile[0] == TEXT('/') && FPackageName::IsValidLongPackageName(SCFile, /*bIncludeReadOnlyRoots*/false))//	if (SCFile[0] == TEXT('/') && FPackageName::IsValidPath(SCFile))
+	if (SCFile[0] == TEXT('/') && FPackageName::IsValidLongPackageName(SCFile, /*bIncludeReadOnlyRoots*/false))
 	{
 		// Assume it is a package
 		bool bPackage = true;
@@ -1007,7 +1007,7 @@ bool USourceControlHelpers::RevertUnchangedFiles(const TArray<FString>& InFiles,
 }
 
 
-bool USourceControlHelpers::CheckInFile(const FString& InFile, const FString& InDescription, bool bSilent)
+bool USourceControlHelpers::CheckInFile(const FString& InFile, const FString& InDescription, bool bSilent, bool bKeepCheckedOut)
 {
 	// Determine file type and ensure it is in form source control wants
 	FString SCFile = SourceControlHelpersInternal::ConvertFileToQualifiedPath(InFile, bSilent);
@@ -1028,13 +1028,15 @@ bool USourceControlHelpers::CheckInFile(const FString& InFile, const FString& In
 	TSharedRef<FCheckIn, ESPMode::ThreadSafe> CheckInOp = ISourceControlOperation::Create<FCheckIn>();
 	CheckInOp->SetDescription(FText::FromString(InDescription));
 
+	CheckInOp->SetKeepCheckedOut(bKeepCheckedOut);
+
 	ECommandResult::Type Result = Provider->Execute(CheckInOp, SCFile);
 
 	return Result == ECommandResult::Succeeded;
 }
 
 
-bool USourceControlHelpers::CheckInFiles(const TArray<FString>& InFiles, const FString& InDescription, bool bSilent)
+bool USourceControlHelpers::CheckInFiles(const TArray<FString>& InFiles, const FString& InDescription, bool bSilent, bool bKeepCheckedOut)
 {
 	// If we have nothing to process, exit immediately
 	if (InFiles.IsEmpty())
@@ -1059,6 +1061,8 @@ bool USourceControlHelpers::CheckInFiles(const TArray<FString>& InFiles, const F
 	TSharedRef<FCheckIn, ESPMode::ThreadSafe> CheckInOp = ISourceControlOperation::Create<FCheckIn>();
 	CheckInOp->SetDescription(FText::FromString(InDescription));
 
+	CheckInOp->SetKeepCheckedOut(bKeepCheckedOut);
+
 	ECommandResult::Type Result = Provider->Execute(CheckInOp, FilePaths);
 
 	return !bFilesSkipped && (Result == ECommandResult::Succeeded);
@@ -1177,7 +1181,43 @@ FSourceControlState USourceControlHelpers::QueryFileState(const FString& InFile,
 	return State;
 }
 
+bool USourceControlHelpers::GetFilesInDepotAtPath(const FString& PathToDirectory, TArray<FString>& OutFilesList, bool bIncludeDeleted, bool bSilent)
+{
+	TRACE_CPUPROFILER_EVENT_SCOPE(USourceControlHelpers::GetFilesInDepotAtPath);
+
+	bool bSuccess = false;
+	if (ISourceControlModule::Get().IsEnabled())
+	{
+		ISourceControlProvider& SourceControlProvider = ISourceControlModule::Get().GetProvider();
+		if (SourceControlProvider.IsAvailable())
+		{
+			FString CorrectedPath = PathToDirectory.Replace(TEXT("/Game"), TEXT(""), ESearchCase::CaseSensitive);
+			FString FullPath = FPaths::ProjectContentDir() / CorrectedPath;
+			FPaths::RemoveDuplicateSlashes(FullPath);
+
+			TArray<FString> FileArray;
+			FileArray.Add(FullPath);
+
+			TSharedRef<FGetFileList, ESPMode::ThreadSafe> Operation = ISourceControlOperation::Create<FGetFileList>();
+			Operation->SetIncludeDeleted(bIncludeDeleted);
+
+			ECommandResult::Type Result = SourceControlProvider.Execute(Operation, FileArray, EConcurrency::Synchronous);
+			bSuccess = (Result == ECommandResult::Succeeded);
 
+			if (!bSuccess)
+			{
+				FFormatNamedArguments Arguments;
+				Arguments.Add(TEXT("PathToDirectory"), FText::FromString(PathToDirectory));
+				SourceControlHelpersInternal::LogError(FText::Format(LOCTEXT("CouldNotGetFileList", "Could not get file list under path: {PathToDirectory}."), Arguments), bSilent);
+			}
+			else
+			{
+				OutFilesList = Operation->GetFilesList();
+			}
+		}
+	}
+	return bSuccess;
+}
 
 static FString PackageFilename_Internal( const FString& InPackageName )
 {
diff --git a/Engine/Source/Developer/SourceControl/Public/SourceControlHelpers.h b/Engine/Source/Developer/SourceControl/Public/SourceControlHelpers.h
--- a/Engine/Source/Developer/SourceControl/Public/SourceControlHelpers.h
+++ b/Engine/Source/Developer/SourceControl/Public/SourceControlHelpers.h
@@ -383,10 +383,11 @@ public:
 	 * @param	InFile			The file to check in - can be either fully qualified path, relative path, long package name, asset path or export text path (often stored on clipboard)
 	 * @param	InDescription	Description for check in
 	 * @param	bSilent			if false (default) then write out any error info to the Log. Any error text can be retrieved by LastErrorMsg() regardless.
+	 * @param	bKeepCheckedOut Keep files checked-out after checking in. This is helpful for maintaining "ownership" of files if further operations are needed.
 	 * @return	true if succeeded, false if failed and can call LastErrorMsg() for more info.
 	 */
 	UFUNCTION(BlueprintCallable, Category = "Editor Scripting | Editor Source Control Helpers")
-	static bool CheckInFile(const FString& InFile, const FString& InDescription, bool bSilent = false);
+	static bool CheckInFile(const FString& InFile, const FString& InDescription, bool bSilent = false, bool bKeepCheckedOut = false);
 
 	/**
 	 * Use currently set source control provider to check in specified files.
@@ -395,10 +396,11 @@ public:
 	 * @param	InFiles			Files to check out - can be either fully qualified path, relative path, long package name, asset path or export text path (often stored on clipboard)
 	 * @param	InDescription	Description for check in
 	 * @param	bSilent			if false (default) then write out any error info to the Log. Any error text can be retrieved by LastErrorMsg() regardless.
+	 * @param	bKeepCheckedOut Keep files checked-out after checking in. This is helpful for maintaining "ownership" of files if further operations are needed.
 	 * @return	true if succeeded, false if failed and can call LastErrorMsg() for more info.
 	 */
 	UFUNCTION(BlueprintCallable, Category = "Editor Scripting | Editor Source Control Helpers")
-	static bool CheckInFiles(const TArray<FString>& InFiles, const FString& InDescription, bool bSilent = false);
+	static bool CheckInFiles(const TArray<FString>& InFiles, const FString& InDescription, bool bSilent = false, bool bKeepCheckedOut = false);
 
 	/**
 	 * Use currently set source control provider to copy a file.
@@ -436,6 +438,17 @@ public:
 	UFUNCTION(BlueprintCallable, Category = "Editor Scripting | Editor Source Control Helpers")
 	static FSourceControlState QueryFileState(const FString& InFile, bool bSilent = false);
 
+	/**
+	 * Use currently set source control provider to query the list of files in the depot under a certain path.
+	 * @note	Blocks until action is complete.
+	 *
+	 * @param	PathToDirectory	The path which we want to query the list of files from.
+	 * @param	OutFilesList	An array containing the list of files under the queried path.
+	 * @param	bIncludeDeleted	Include files that have been deleted from the depot.
+	 * @param	bSilent			if false (default) then write out any error info to the Log. Any error text can be retrieved by LastErrorMsg() regardless.
+	 * @return	Success or failure of the operation
+	 */
+	static bool GetFilesInDepotAtPath(const FString& PathToDirectory, TArray<FString>& OutFilesList, bool bIncludeDeleted = false, bool bSilent = false);
 
 	/**
 	 * Helper function to get a filename for a package name.
diff --git a/Engine/Source/Developer/SourceControl/Public/SourceControlOperations.h b/Engine/Source/Developer/SourceControl/Public/SourceControlOperations.h
--- a/Engine/Source/Developer/SourceControl/Public/SourceControlOperations.h
+++ b/Engine/Source/Developer/SourceControl/Public/SourceControlOperations.h
@@ -94,12 +94,25 @@ public:
 		return SuccessMessage;
 	}
 
+	void SetKeepCheckedOut( const bool bInKeepCheckedOut )
+	{
+		bKeepCheckedOut = bInKeepCheckedOut;
+	}
+
+	bool GetKeepCheckedOut()
+	{
+		return bKeepCheckedOut;
+	}
+
 protected:
 	/** Description of the checkin */
 	FText Description;
 
 	/** A short message listing changelist/revision we submitted, if successful */
 	FText SuccessMessage;
+
+	/** Keep files checked-out after checking in */
+	bool bKeepCheckedOut = false;
 };
 
 /**
@@ -120,6 +133,51 @@ public:
 	}
 };
 
+/**
+ * Operation used to get the file list of a folder out of source control
+ */
+class FGetFileList : public FSourceControlOperationBase
+{
+public:
+	// ISourceControlOperation interface
+	virtual FName GetName() const override
+	{
+		return "GetFileList";
+	}
+
+	virtual FText GetInProgressString() const override
+	{
+		return LOCTEXT("SourceControl_GetFileList", "Getting file list out of Source Control...");
+	}
+
+	void SetIncludeDeleted( const bool bInIncludeDeleted )
+	{
+		bIncludeDeleted = bInIncludeDeleted;
+	}
+
+	bool GetIncludeDeleted()
+	{
+		return bIncludeDeleted;
+	}
+
+	const TArray<FString>& GetFilesList() const
+	{
+		return FilesList;
+	}
+
+	void SetFilesList(TArray<FString>&& InFilesList)
+	{
+		FilesList = MoveTemp(InFilesList);
+	}
+
+protected:
+	/** Include deleted files in the list. */
+	bool bIncludeDeleted = false;
+
+	/** Stored result of the operation */
+	TArray<FString> FilesList;
+};
+
 /**
  * Operation used to mark files for add in source control
  */
