Added option to provide a different project directory to be used by source control.

diff --git a/Engine/Plugins/Developer/PerforceSourceControl/Source/PerforceSourceControl/Private/PerforceConnection.cpp b/Engine/Plugins/Developer/PerforceSourceControl/Source/PerforceSourceControl/Private/PerforceConnection.cpp
--- a/Engine/Plugins/Developer/PerforceSourceControl/Source/PerforceSourceControl/Private/PerforceConnection.cpp
+++ b/Engine/Plugins/Developer/PerforceSourceControl/Source/PerforceSourceControl/Private/PerforceConnection.cpp
@@ -663,7 +663,7 @@ bool FPerforceConnection::GetWorkspaceList(const FPerforceConnectionInfo& InConn
 
 		if (bCommandOK)
 		{
-			FString ApplicationPath = IFileManager::Get().ConvertToAbsolutePathForExternalAppForRead(*FPaths::ProjectDir()).ToLower();
+			FString ApplicationPath = IFileManager::Get().ConvertToAbsolutePathForExternalAppForRead(*ISourceControlModule::Get().GetSourceControlProjectDir()).ToLower();
 
 			FString LocalHostName = InConnectionInfo.HostOverride;
 			if(LocalHostName.Len() == 0)
diff --git a/Engine/Plugins/Developer/PerforceSourceControl/Source/PerforceSourceControl/Private/PerforceSourceControlOperations.cpp b/Engine/Plugins/Developer/PerforceSourceControl/Source/PerforceSourceControl/Private/PerforceSourceControlOperations.cpp
--- a/Engine/Plugins/Developer/PerforceSourceControl/Source/PerforceSourceControl/Private/PerforceSourceControlOperations.cpp
+++ b/Engine/Plugins/Developer/PerforceSourceControl/Source/PerforceSourceControl/Private/PerforceSourceControlOperations.cpp
@@ -7,6 +7,7 @@
 #include "Misc/Paths.h"
 #include "Misc/EngineVersion.h"
 #include "Modules/ModuleManager.h"
+#include "ISourceControlModule.h"
 #include "SourceControlHelpers.h"
 #include "SourceControlOperations.h"
 #include "PerforceSourceControlRevision.h"
@@ -220,7 +221,7 @@ static bool UpdateCachedStates(const TMap<FString, EPerforceState::Type>& InResu
 
 static bool CheckWorkspaceRecordSet(const FP4RecordSet& InRecords, TArray<FText>& OutErrorMessages, FText& OutNotificationText)
 {
-	FString ApplicationPath = IFileManager::Get().ConvertToAbsolutePathForExternalAppForRead(*FPaths::ProjectDir()).ToLower();
+	FString ApplicationPath = IFileManager::Get().ConvertToAbsolutePathForExternalAppForRead(*ISourceControlModule::Get().GetSourceControlProjectDir()).ToLower();
 	ApplicationPath = ApplicationPath.Replace(TEXT("\\"), TEXT("/"));
 
 	for(const auto& Record : InRecords)
diff --git a/Engine/Source/Developer/SourceControl/Private/SourceControlModule.cpp b/Engine/Source/Developer/SourceControl/Private/SourceControlModule.cpp
--- a/Engine/Source/Developer/SourceControl/Private/SourceControlModule.cpp
+++ b/Engine/Source/Developer/SourceControl/Private/SourceControlModule.cpp
@@ -6,6 +6,7 @@
 #include "SourceControlOperations.h"
 #include "SourceControlHelpers.h"
 #include "SourceControlAssetDataCache.h"
+#include "Misc/Paths.h"
 
 #if SOURCE_CONTROL_WITH_SLATE
 	#include "Widgets/DeclarativeSyntaxSupport.h"
@@ -506,6 +507,29 @@ const FSourceControlFilesDeletedDelegate& FSourceControlModule::GetOnFilesDelete
 	return OnFilesDeleted;
 }
 
+void FSourceControlModule::RegisterSourceControlProjectDirDelegate(const FSourceControlProjectDirDelegate& InSourceControlProjectDirDelegate)
+{
+	SourceControlProjectDirDelegate = InSourceControlProjectDirDelegate;
+}
+
+void FSourceControlModule::UnregisterSourceControlProjectDirDelegate()
+{
+	SourceControlProjectDirDelegate = FSourceControlProjectDirDelegate();
+}
+
+FString FSourceControlModule::GetSourceControlProjectDir() const
+{
+	if (SourceControlProjectDirDelegate.IsBound())
+	{
+		FString ProjectDir = SourceControlProjectDirDelegate.Execute();
+		if (!ProjectDir.IsEmpty())
+		{
+			return ProjectDir;
+		}
+	}
+	return FPaths::ProjectDir();
+}
+
 IMPLEMENT_MODULE( FSourceControlModule, SourceControl );
 
 #undef LOCTEXT_NAMESPACE
diff --git a/Engine/Source/Developer/SourceControl/Private/SourceControlModule.h b/Engine/Source/Developer/SourceControl/Private/SourceControlModule.h
--- a/Engine/Source/Developer/SourceControl/Private/SourceControlModule.h
+++ b/Engine/Source/Developer/SourceControl/Private/SourceControlModule.h
@@ -53,6 +53,10 @@ public:
 	virtual void UnregisterFilesDeleted(FDelegateHandle InHandle) override;
 	virtual const FSourceControlFilesDeletedDelegate& GetOnFilesDeleted() const override;
 
+	virtual void RegisterSourceControlProjectDirDelegate(const FSourceControlProjectDirDelegate& SourceControlProjectDirDelegate) override;
+	virtual void UnregisterSourceControlProjectDirDelegate() override;
+	virtual FString GetSourceControlProjectDir() const override;
+
 	/** Save the settings to the ini file */
 	void SaveSettings();
 
@@ -144,4 +148,7 @@ private:
 
 	/** Used to cache source controlled AssetData information */
 	FSourceControlAssetDataCache AssetDataCache;
+
+	/** Delegate used to return the current project base directory */
+	FSourceControlProjectDirDelegate SourceControlProjectDirDelegate;
 };
diff --git a/Engine/Source/Developer/SourceControl/Public/ISourceControlModule.h b/Engine/Source/Developer/SourceControl/Public/ISourceControlModule.h
--- a/Engine/Source/Developer/SourceControl/Public/ISourceControlModule.h
+++ b/Engine/Source/Developer/SourceControl/Public/ISourceControlModule.h
@@ -23,6 +23,9 @@ DECLARE_MULTICAST_DELEGATE_TwoParams( FSourceControlProviderChanged, ISourceCont
 /** Delegate called on pre-submit for data validation */
 DECLARE_DELEGATE_FourParams(FSourceControlPreSubmitDataValidationDelegate, FSourceControlChangelistPtr /*Changelist*/, EDataValidationResult& /*Result*/, TArray<FText>& /*ValidationErrors*/, TArray<FText>& /*ValidationWarnings*/);
 
+/** Delegate used to specify the project base directory to be used by the source control */
+DECLARE_DELEGATE_RetVal(FString, FSourceControlProjectDirDelegate);
+
 /** 
  * Delegate called once the user has confirmed that they want to submit files to source control BUT before the files are actually submitted.
  * It is intended for last minute checks that can only run once there is no chance of the user canceling the actual submit.
@@ -206,6 +209,22 @@ public:
 	 */
 	virtual const FSourceControlFilesDeletedDelegate& GetOnFilesDeleted() const = 0;
 
+	/**
+	 * Register a delegate used to specify the project base directory to be used by the source control
+	 */
+	virtual void RegisterSourceControlProjectDirDelegate(const FSourceControlProjectDirDelegate& SourceControlProjectDirDelegate) = 0;
+
+	/**
+	 * Unregister the FSourceControlProjectDirDelegate delegate
+	 */
+	virtual void UnregisterSourceControlProjectDirDelegate() = 0;
+
+
+	/**
+	 * Returns the project base directory to be used by the source control
+	 */
+	virtual FString GetSourceControlProjectDir() const = 0;
+
 	/**
 	 * Gets a reference to the source control module instance.
 	 *
