The source control operation FCreateWorkspace now supports setting a description.

#rnx

diff --git a/Engine/Plugins/Developer/PerforceSourceControl/Source/PerforceSourceControl/Private/PerforceSourceControlOperations.cpp b/Engine/Plugins/Developer/PerforceSourceControl/Source/PerforceSourceControl/Private/PerforceSourceControlOperations.cpp
--- a/Engine/Plugins/Developer/PerforceSourceControl/Source/PerforceSourceControl/Private/PerforceSourceControlOperations.cpp
+++ b/Engine/Plugins/Developer/PerforceSourceControl/Source/PerforceSourceControl/Private/PerforceSourceControlOperations.cpp
@@ -2937,6 +2937,12 @@ bool FPerforceCreateWorkspaceWorker::Execute(class FPerforceSourceControlCommand
 
 		ClientDesc << TEXT("Client:\t") << Operation->GetWorkspaceName() << TEXT("\n\n");
 		ClientDesc << TEXT("Owner:\t") << Connection.GetUser() << TEXT("\n\n");
+
+		if (!Operation->GetWorkspaceDescription().IsEmpty())
+		{
+			ClientDesc << TEXT("Description:\t") << Operation->GetWorkspaceDescription() << TEXT("\n\n");
+		}
+
 		ClientDesc << TEXT("Root:\t") << Operation->GetWorkspaceRoot() << TEXT("\n\n");
 		ClientDesc << TEXT("View:\t\n\n");
 
diff --git a/Engine/Source/Developer/SourceControl/Public/SourceControlOperations.h b/Engine/Source/Developer/SourceControl/Public/SourceControlOperations.h
--- a/Engine/Source/Developer/SourceControl/Public/SourceControlOperations.h
+++ b/Engine/Source/Developer/SourceControl/Public/SourceControlOperations.h
@@ -857,11 +857,17 @@ public:
 	/**
 	 * Constructor
 	 * 
-	 * @param InWorkspaceName	The name of the workspace to create
-	 * @param InWorkspaceRoom	The file path to the workspace root (can be relative to the project)
+	 * @param WorkspaceName	The name of the workspace to create
+	 * @param WorkspaceRoom	The file path to the workspace root (can be relative to the project)
 	 */
-	SOURCECONTROL_API FCreateWorkspace(FStringView InWorkspaceName, FStringView InWorkspaceRoot);
+	SOURCECONTROL_API FCreateWorkspace(FStringView WorkspaceName, FStringView WorkspaceRoot);
 	
+	/** Set the description to be used by the workspace, if left unset then the default description will be used */
+	void SetDescription(FStringView Desciption)
+	{
+		WorkspaceDescription = Desciption;
+	}
+
 	/**
 	 * Add a new mapping for the client spec in the native format of the current source control provider.
 	 * These will be written to the client spec in the order that they are added.
@@ -874,6 +880,7 @@ public:
 		ClientView.Emplace(DepotPath, ClientPath);
 	}
 
+	/** Set the type of workspace to be created. @see FCreateWorkspace::EType */
 	void SetType(EType InType)
 	{
 		Type = InType;
@@ -889,6 +896,11 @@ public:
 		return WorkspaceRoot;
 	}
 
+	const FString& GetWorkspaceDescription() const
+	{
+		return WorkspaceDescription;
+	}
+
 	const TArray<FClientViewMapping> GetClientView() const
 	{
 		return ClientView;
@@ -916,6 +928,8 @@ protected:
 	FString WorkspaceName;
 	FString WorkspaceRoot;
 
+	FString WorkspaceDescription;
+
 	EType Type = EType::Writeable;
 
 	TArray<FClientViewMapping> ClientView;
