Modifies UpdateStatus to enforce full validation only if bForceUpdate flag is used. Ignoring files not opened/edited by default.

diff --git a/Engine/Plugins/Developer/PerforceSourceControl/Source/PerforceSourceControl/Private/PerforceSourceControlOperations.cpp b/Engine/Plugins/Developer/PerforceSourceControl/Source/PerforceSourceControl/Private/PerforceSourceControlOperations.cpp
--- a/Engine/Plugins/Developer/PerforceSourceControl/Source/PerforceSourceControl/Private/PerforceSourceControlOperations.cpp
+++ b/Engine/Plugins/Developer/PerforceSourceControl/Source/PerforceSourceControl/Private/PerforceSourceControlOperations.cpp
@@ -1752,7 +1752,13 @@ bool FPerforceUpdateStatusWorker::Execute(FPerforceSourceControlCommand& InComma
 			TArray<FString> Parameters;
 			FP4RecordSet Records;
 			// Query for open files different than the versions stored in Perforce
-			Parameters.Add(TEXT("-fsa"));
+
+			if (Operation->ShouldForceUpdate())
+			{
+				Parameters.Add(TEXT("-f"));
+			}
+
+			Parameters.Add(TEXT("-sa"));
 			for (FString File : InCommand.Files)
 			{
 				if (IFileManager::Get().DirectoryExists(*File))
diff --git a/Engine/Source/Developer/SourceControl/Public/SourceControlOperations.h b/Engine/Source/Developer/SourceControl/Public/SourceControlOperations.h
--- a/Engine/Source/Developer/SourceControl/Public/SourceControlOperations.h
+++ b/Engine/Source/Developer/SourceControl/Public/SourceControlOperations.h
@@ -346,6 +346,7 @@ public:
 		, bUpdateModifiedState(false)
 		, bCheckingAllFiles(false)
 		, bForceQuiet(false)
+		, bForceUpdate(false)
 	{
 	}
 
@@ -385,6 +386,11 @@ public:
 		bForceQuiet = bInQuiet;
 	}
 
+	void SetForceUpdate(const bool bInForceUpdate)
+	{
+		bForceUpdate = bInForceUpdate;
+	}
+
 	bool ShouldUpdateHistory() const
 	{
 		return bUpdateHistory;
@@ -410,6 +416,11 @@ public:
 		return bForceQuiet;
 	}
 
+	bool ShouldForceUpdate() const
+	{
+		return bForceUpdate;
+	}
+
 protected:
 	/** Whether to update history */
 	bool bUpdateHistory;
@@ -425,6 +436,9 @@ protected:
 
 	/** Controls whether the operation will trigger an update or not */
 	bool bForceQuiet;
+
+	/** Forces the verification for provided files - providers can ignore files not opened/edited without it */
+	bool bForceUpdate;
 };
 
 /**
diff --git a/Engine/Source/Developer/UncontrolledChangelists/Private/UncontrolledChangelistState.cpp b/Engine/Source/Developer/UncontrolledChangelists/Private/UncontrolledChangelistState.cpp
--- a/Engine/Source/Developer/UncontrolledChangelists/Private/UncontrolledChangelistState.cpp
+++ b/Engine/Source/Developer/UncontrolledChangelists/Private/UncontrolledChangelistState.cpp
@@ -118,6 +118,7 @@ bool FUncontrolledChangelistState::AddFiles(const TArray<FString>& InFilenames,
 		auto UpdateStatusOperation = ISourceControlOperation::Create<FUpdateStatus>();
 		UpdateStatusOperation->SetUpdateModifiedState(true);
 		UpdateStatusOperation->SetQuiet(true);
+		UpdateStatusOperation->SetForceUpdate(true);
 		SourceControlProvider.Execute(UpdateStatusOperation, InFilenames);
 	}
 
