Allow the caller to decide if FUpdateStatus will check for directy paths based on the path itself (terminating in a seperator) or if the file system should be used instead.

#rnx

- The problem we have is if a p4 depot path is used when requesting filestates (to see if the file exists in p4 for example) it will be treated as a networkpath due to it starting with // so calling the filesystem to check if it is a directory or not can get very expensive.
- Other operations just check the end of the path for '/' or '\' which is what this operation used to do, but that was changed at some point to check the file system, presumably to fix another problem and as we don't really know how 3rd party code is using this system we cannot just change the behavior without risking breaking existing code bases.
- This fix allows the caller to opt in to the old behavior of just checking the end of the path to determine if it is a directory or not.
- This fix is not very good, future work for a better fix will be tracked in the attached jira.

diff --git a/Engine/Plugins/Developer/PerforceSourceControl/Source/PerforceSourceControl/Private/PerforceSourceControlOperations.cpp b/Engine/Plugins/Developer/PerforceSourceControl/Source/PerforceSourceControl/Private/PerforceSourceControlOperations.cpp
--- a/Engine/Plugins/Developer/PerforceSourceControl/Source/PerforceSourceControl/Private/PerforceSourceControlOperations.cpp
+++ b/Engine/Plugins/Developer/PerforceSourceControl/Source/PerforceSourceControl/Private/PerforceSourceControlOperations.cpp
@@ -1712,7 +1712,7 @@ bool FPerforceUpdateStatusWorker::Execute(FPerforceSourceControlCommand& InComma
 			// Mandatory parameters (the list of files to stat):
 			for (FString File : InCommand.Files)
 			{
-				if (IFileManager::Get().DirectoryExists(*File))
+				if (Operation->IsDirectoryPath(File))
 				{
 					// If the file is a directory, do a recursive fstat on the contents
 					File /= TEXT("...");
diff --git a/Engine/Source/Developer/SourceControl/Private/SourceControlOperations.cpp b/Engine/Source/Developer/SourceControl/Private/SourceControlOperations.cpp
--- a/Engine/Source/Developer/SourceControl/Private/SourceControlOperations.cpp
+++ b/Engine/Source/Developer/SourceControl/Private/SourceControlOperations.cpp
@@ -3,10 +3,47 @@
 #include "SourceControlOperations.h"
 
 #include "Containers/StringView.h"
+#include "HAL/FileManager.h"
 #include "ISourceControlModule.h"
 #include "Misc/PathViews.h"
 #include "Misc/Paths.h"
 
+bool FUpdateStatus::IsDirectoryPath(const FString& Path) const
+{
+	// Problem:
+	// We need to be able to work out if a path if a directory or not as some source 
+	// control implementations will need to append wildcards to the path if it is
+	// a directory in order for it to function correctly.
+	// 
+	// Strictly speaking the only way we can tell the difference between a file path
+	// and a directory path is to check with the file system. 
+	// However some source control implementations let us reference files on the server
+	// and do not require us to have them on the users machine. This means checking with
+	// the file system does not give an accurate result and can end up being very 
+	// slow if the file system interprets the path as a UNC path and tries to resolve it.
+	//
+	// Ideally we would let the caller clearly mark if the path is for a file or a 
+	// directory but this would require a large overhaul of the source control API.
+	// We cannot change the existing behavior as it is very likely that 3rd party code is
+	// relying on it.
+	//
+	// Solution:
+	// As a 'temporary' solution, the caller can set bSetRequireDirPathEndWithSeparator
+	// to true. When it is true we will assume that all directory paths end with a 
+	// separator. Although we tend to not terminate directory paths with a separator in the 
+	// engine, this will allow us an opt in to the new behavior until more substantial API
+	// changes can be made.
+
+	if (bSetRequireDirPathEndWithSeparator)
+	{
+		return Path.EndsWith(TEXT("/")) || Path.EndsWith(TEXT("\\"));
+	}
+	else
+	{
+		return IFileManager::Get().DirectoryExists(*Path);
+	}
+}
+
 FDownloadFile::FDownloadFile(FStringView InTargetDirectory, EVerbosity InVerbosity)
 	: Verbosity(InVerbosity)
 	, TargetDirectory(InTargetDirectory)
diff --git a/Engine/Source/Developer/SourceControl/Public/SourceControlOperations.h b/Engine/Source/Developer/SourceControl/Public/SourceControlOperations.h
--- a/Engine/Source/Developer/SourceControl/Public/SourceControlOperations.h
+++ b/Engine/Source/Developer/SourceControl/Public/SourceControlOperations.h
@@ -347,6 +347,7 @@ public:
 		, bCheckingAllFiles(false)
 		, bForceQuiet(false)
 		, bForceUpdate(false)
+		, bSetRequireDirPathEndWithSeparator(false)
 	{
 	}
 
@@ -386,6 +387,22 @@ public:
 		bForceQuiet = bInQuiet;
 	}
 
+	/** 
+	 * Sets the method that the operation will use to determine if a path
+	 * references a file or a directory. For more details @see IsDirectoryPath()
+	 * 
+	 * @param bFlag When true the operation will check the path and assume that it
+	 *				is a directory if the path ends with '/' or '\' and a file if it 
+	 *				does not.
+	 *				When false (the default) the operation will poll the file system
+	 *				with the path to see if it is a file or a directory.
+	 * 
+	 */
+	void SetRequireDirPathEndWithSeperator(bool bFlag)
+	{
+		bSetRequireDirPathEndWithSeparator = bFlag;
+	}
+
 	void SetForceUpdate(const bool bInForceUpdate)
 	{
 		bForceUpdate = bInForceUpdate;
@@ -421,6 +438,19 @@ public:
 		return bForceUpdate;
 	}
 
+	/** 
+	 * Returns if the given path should be considered a directory or not.
+	 * If bSetRequireDirPathEndWithSeparator is not set (the default) then
+	 * we will poll the file system. However in some cases this can be very
+	 * slow, in which case bSetRequireDirPathEndWithSeparator can be set 
+	 * to true and we require that any directory path must be terminated
+	 * by a path separator (/ or \) so that we can tell by simply looking at
+	 * the path itself.
+	 * This is opt in behavior to avoid breaking existing 3rd party code that
+	 * relies on the default behavior.
+	 */
+	SOURCECONTROL_API bool IsDirectoryPath(const FString& Path) const;
+
 protected:
 	/** Whether to update history */
 	bool bUpdateHistory;
@@ -439,6 +469,9 @@ protected:
 
 	/** Forces the verification for provided files - providers can ignore files not opened/edited without it */
 	bool bForceUpdate;
+
+	/** If we should assume paths ending in a separator are directory paths or do we need to check with the file system? */
+	bool bSetRequireDirPathEndWithSeparator;
 };
 
 /**
