Fix typo in method name

#rnx

- Fix incorrect spelling in the method name.
- Oddly enough this was spelt correctly for both the member and in the comments.

diff --git a/Engine/Source/Developer/SourceControl/Public/SourceControlOperations.h b/Engine/Source/Developer/SourceControl/Public/SourceControlOperations.h
--- a/Engine/Source/Developer/SourceControl/Public/SourceControlOperations.h
+++ b/Engine/Source/Developer/SourceControl/Public/SourceControlOperations.h
@@ -398,7 +398,7 @@ public:
 	 *				with the path to see if it is a file or a directory.
 	 * 
 	 */
-	void SetRequireDirPathEndWithSeperator(bool bFlag)
+	void SetRequireDirPathEndWithSeparator(bool bFlag)
 	{
 		bSetRequireDirPathEndWithSeparator = bFlag;
 	}
