Added an array view overload to FUpdatePendingChangelistsStatus::SetChangelistsToUpdate to make it easier for the caller to add a single changelist

#rnx

diff --git a/Engine/Source/Developer/SourceControl/Public/SourceControlOperations.h b/Engine/Source/Developer/SourceControl/Public/SourceControlOperations.h
--- a/Engine/Source/Developer/SourceControl/Public/SourceControlOperations.h
+++ b/Engine/Source/Developer/SourceControl/Public/SourceControlOperations.h
@@ -585,6 +585,12 @@ public:
 		bUpdateAllChangelists = false;
 	}
 
+	void SetChangelistsToUpdate(const TArrayView<FSourceControlChangelistRef>& InChangelistsToUpdate)
+	{
+		ChangelistsToUpdate = InChangelistsToUpdate;
+		bUpdateAllChangelists = false;
+	}
+
 	const TArray<FSourceControlChangelistRef>& GetChangelistsToUpdate() const
 	{
 		return ChangelistsToUpdate;
