SourceControl: Add FGetSubmittedChangelists operation.

diff --git a/Engine/Source/Developer/SourceControl/Public/SourceControlOperations.h b/Engine/Source/Developer/SourceControl/Public/SourceControlOperations.h
--- a/Engine/Source/Developer/SourceControl/Public/SourceControlOperations.h
+++ b/Engine/Source/Developer/SourceControl/Public/SourceControlOperations.h
@@ -542,6 +542,24 @@ public:
 	}
 };
 
+/**
+ * Operation used to retrieve submitted changelist(s).
+ */
+class FGetSubmittedChangelists : public FSourceControlOperationBase
+{
+public:
+	// ISourceControlOperation interface
+	virtual FName GetName() const override
+	{
+		return "GetSubmittedChangelists"; 
+	}
+
+	virtual FText GetInProgressString() const override
+	{
+		return LOCTEXT("SourceControl_GetSubmittedChangelists", "Retrieving submitted changelist(s) from Source Control...");
+	}
+};
+
 /**
  * This operations query the source control to extract all the details available for a given changelist. The operations returns a collection of key/value corresponding to the details available. 
  * The list of key/value is specific to the source control implementation.
