Updates to source control helpers that:
- allow for reverting files marked for delete
- return success on a successful 'mark for delete'
#rnx

diff --git a/Engine/Source/Developer/SourceControl/Private/SourceControlHelpers.cpp b/Engine/Source/Developer/SourceControl/Private/SourceControlHelpers.cpp
--- a/Engine/Source/Developer/SourceControl/Private/SourceControlHelpers.cpp
+++ b/Engine/Source/Developer/SourceControl/Private/SourceControlHelpers.cpp
@@ -739,8 +739,6 @@ bool USourceControlHelpers::MarkFileForDelete(const FString& InFile, bool bSilen
 		return false;
 	}
 
-	bool bDelete = false;
-
 	if (SCState->IsSourceControlled())
 	{
 		bool bAdded = SCState->IsAdded();
@@ -782,8 +780,10 @@ bool USourceControlHelpers::MarkFileForDelete(const FString& InFile, bool bSilen
 		// Don't bother checking if it exists since Delete doesn't care.
 		return FileManager.Delete(*SCFile, false, true);
 	}
-
-	return false;
+	else
+	{
+		return true; 
+	}
 }
 
 
@@ -930,7 +930,7 @@ bool USourceControlHelpers::RevertFiles(const TArray<FString>& InFiles,	bool bSi
 
 		// Less error checking and info is made for multiple files than the single file version.
 		// This multi-file version could be made similarly more sophisticated.
-		if (SCState->IsCheckedOut() || SCState->IsAdded())
+		if (SCState->CanRevert())
 		{
 			SCFilesToRevert.Add(SCFile);
 		}
@@ -1104,7 +1104,6 @@ bool USourceControlHelpers::CopyFile(const FString& InSourcePath, const FString&
 	return Result == ECommandResult::Succeeded;
 }
 
-
 FSourceControlState USourceControlHelpers::QueryFileState(const FString& InFile, bool bSilent)
 {
 	FSourceControlState State;
