Adding Async version of QueryFileState.

diff --git a/Engine/Source/Developer/SourceControl/Private/SourceControlHelpers.cpp b/Engine/Source/Developer/SourceControl/Private/SourceControlHelpers.cpp
--- a/Engine/Source/Developer/SourceControl/Private/SourceControlHelpers.cpp
+++ b/Engine/Source/Developer/SourceControl/Private/SourceControlHelpers.cpp
@@ -1180,6 +1180,72 @@ FSourceControlState USourceControlHelpers::QueryFileState(const FString& InFile,
 	return State;
 }
 
+void USourceControlHelpers::AsyncQueryFileState(FQueryFileStateDelegate FileStateCallback, const FString& InFile, bool bSilent)
+{
+	// Determine file type and ensure it is in form source control wants
+	FString SCFile = SourceControlHelpersInternal::ConvertFileToQualifiedPath(InFile, bSilent);
+
+	if (SCFile.IsEmpty())
+	{
+		return;
+	}
+
+	if (ISourceControlProvider* Provider = SourceControlHelpersInternal::VerifySourceControl(bSilent))
+	{
+		// Make sure we update the modified state of the files (Perforce requires this
+		// since can be a more expensive test).
+		TSharedRef<FUpdateStatus, ESPMode::ThreadSafe> UpdateStatusOperation = ISourceControlOperation::Create<FUpdateStatus>();
+		UpdateStatusOperation->SetUpdateModifiedState(true);
+
+		ISourceControlModule::Get().GetProvider().Execute(UpdateStatusOperation, SCFile, EConcurrency::Asynchronous,
+			FSourceControlOperationComplete::CreateLambda([FileStateCallback, Provider, SCFile, InFile, bSilent](const FSourceControlOperationRef& InOperation, ECommandResult::Type InResult)
+		{
+			FSourceControlState State;
+			const FSourceControlStatePtr SCState = Provider->GetState(SCFile, EStateCacheUsage::Use);
+
+			if (!SCState.IsValid())
+			{
+				// Improper or invalid SCC state
+				FFormatNamedArguments Arguments;
+				Arguments.Add(TEXT("InFile"), FText::FromString(InFile));
+				Arguments.Add(TEXT("SCFile"), FText::FromString(SCFile));
+				SourceControlHelpersInternal::LogError(FText::Format(LOCTEXT("CouldNotDetermineState", "Could not determine source control state of file '{InFile}' ({SCFile})."), Arguments), bSilent);
+
+				FileStateCallback.ExecuteIfBound(State);
+				return;
+			}
+
+			State.bIsValid = true;
+
+			// Copy over state info
+			// - make these assignments a method of FSourceControlState if anything else sets a state
+			State.bIsUnknown = SCState->IsUnknown();
+			State.bIsSourceControlled = SCState->IsSourceControlled();
+			State.bCanCheckIn = SCState->CanCheckIn();
+			State.bCanCheckOut = SCState->CanCheckout();
+			State.bIsCheckedOut = SCState->IsCheckedOut();
+			State.bIsCurrent = SCState->IsCurrent();
+			State.bIsAdded = SCState->IsAdded();
+			State.bIsDeleted = SCState->IsDeleted();
+			State.bIsIgnored = SCState->IsIgnored();
+			State.bCanEdit = SCState->CanEdit();
+			State.bCanDelete = SCState->CanDelete();
+			State.bCanAdd = SCState->CanAdd();
+			State.bIsConflicted = SCState->IsConflicted();
+			State.bCanRevert = SCState->CanRevert();
+			State.bIsModified = SCState->IsModified();
+			State.bIsCheckedOutOther = SCState->IsCheckedOutOther();
+
+			if (State.bIsCheckedOutOther)
+			{
+				SCState->IsCheckedOutOther(&State.CheckedOutOther);
+			}
+
+			FileStateCallback.ExecuteIfBound(State);
+		}));
+	}
+}
+
 bool USourceControlHelpers::GetFilesInDepotAtPath(const FString& PathToDirectory, TArray<FString>& OutFilesList, bool bIncludeDeleted, bool bSilent)
 {
 	TRACE_CPUPROFILER_EVENT_SCOPE(USourceControlHelpers::GetFilesInDepotAtPath);
diff --git a/Engine/Source/Developer/SourceControl/Public/SourceControlHelpers.h b/Engine/Source/Developer/SourceControl/Public/SourceControlHelpers.h
--- a/Engine/Source/Developer/SourceControl/Public/SourceControlHelpers.h
+++ b/Engine/Source/Developer/SourceControl/Public/SourceControlHelpers.h
@@ -438,6 +438,18 @@ public:
 	UFUNCTION(BlueprintCallable, Category = "Editor Scripting | Editor Source Control Helpers")
 	static FSourceControlState QueryFileState(const FString& InFile, bool bSilent = false);
 
+	//Delegate to broadcast FileState upon AsyncQueryFileState completion
+	DECLARE_DYNAMIC_DELEGATE_OneParam(FQueryFileStateDelegate, FSourceControlState, FileStateOut);
+	/**
+	* Query the source control state of the specified file, asynchronously.
+	*
+	* @param	FileStateCallback Source control state - see USourceControlState. It will have bIsValid set to false if it could not have its values set.
+	* @param	InFile			  The file to query - can be either fully qualified path, relative path, long package name, asset path or export text path (often stored on clipboard)
+	* @param	bSilent			  if false (default) then write out any error info to the Log. Any error text can be retrieved by LastErrorMsg() regardless.
+	*/
+	UFUNCTION(BlueprintCallable, Category = "Editor Scripting | Editor Source Control Helpers")
+	static void AsyncQueryFileState(FQueryFileStateDelegate FileStateCallback, const FString& InFile, bool bSilent = false);
+
 	/**
 	 * Use currently set source control provider to query the list of files in the depot under a certain path.
 	 * @note	Blocks until action is complete.
