Update the size of SSourceControlLogin to let the user-defined source control plugin auto-adjust its size.

diff --git a/Engine/Source/Developer/SourceControl/Private/SSourceControlLogin.cpp b/Engine/Source/Developer/SourceControl/Private/SSourceControlLogin.cpp
--- a/Engine/Source/Developer/SourceControl/Private/SSourceControlLogin.cpp
+++ b/Engine/Source/Developer/SourceControl/Private/SSourceControlLogin.cpp
@@ -47,7 +47,7 @@ void SSourceControlLogin::Construct(const FArguments& InArgs)
 	ChildSlot
 	[
 		SNew(SBox)
-		.WidthOverride(480.f)
+		.MinDesiredWidth(480.f)
 		[
 			SNew(SBorder)
 			.Padding(0)
