SCC: Dependency validation only on hard refs

diff --git a/Engine/Source/Developer/SourceControl/Private/SourceControlHelpers.cpp b/Engine/Source/Developer/SourceControl/Private/SourceControlHelpers.cpp
--- a/Engine/Source/Developer/SourceControl/Private/SourceControlHelpers.cpp
+++ b/Engine/Source/Developer/SourceControl/Private/SourceControlHelpers.cpp
@@ -1623,7 +1623,7 @@ bool USourceControlHelpers::GetAssetData(const FString & InFileName, const FStri
 		// Assets are already in the cache, we can query dependencies directly
 		if (bGetDependencies)
 		{
-			AssetRegistryModule.Get().GetDependencies(*InPackageName, *OutDependencies);
+			AssetRegistryModule.Get().GetDependencies(*InPackageName, *OutDependencies, UE::AssetRegistry::EDependencyCategory::Package, UE::AssetRegistry::EDependencyQuery::Hard);
 		}
 
 		return true;
