Modifies revert operations in source control window to also reload packages.
Adds Source Control Helper method to perform a modular source control operation and reload the concerned packages

#changelist validated

diff --git a/Engine/Source/Developer/SourceControl/Private/SourceControlHelpers.cpp b/Engine/Source/Developer/SourceControl/Private/SourceControlHelpers.cpp
--- a/Engine/Source/Developer/SourceControl/Private/SourceControlHelpers.cpp
+++ b/Engine/Source/Developer/SourceControl/Private/SourceControlHelpers.cpp
@@ -1,6 +1,7 @@
 // Copyright Epic Games, Inc. All Rights Reserved.
 
 #include "SourceControlHelpers.h"
+#include "Algo/Transform.h"
 #include "ISourceControlState.h"
 #include "HAL/FileManager.h"
 #include "Misc/Paths.h"
@@ -15,8 +16,12 @@
 #include "Logging/MessageLog.h"
 #include "AssetRegistry/AssetRegistryModule.h"
 
-#define LOCTEXT_NAMESPACE "SourceControlHelpers"
+#if WITH_EDITOR
+#include "Editor.h"
+#include "PackageTools.h"
+#endif
 
+#define LOCTEXT_NAMESPACE "SourceControlHelpers"
 
 namespace SourceControlHelpersInternal
 {
@@ -896,6 +901,118 @@ bool USourceControlHelpers::RevertFile(const FString& InFile, bool bSilent)
 	return Result == ECommandResult::Succeeded;
 }
 
+#if WITH_EDITOR
+bool USourceControlHelpers::ApplyOperationAndReloadPackages(const TArray<FString>& InFilenames, const TFunctionRef<bool(const TArray<FString>&)>& InOperation)
+{
+	ISourceControlProvider& SourceControlProvider = ISourceControlModule::Get().GetProvider();
+	TArray<UPackage*> LoadedPackages;
+	TArray<FString> PackageNames;
+	TArray<FString> PackageFilenames;
+	TArray<FString> FilteredActorPackages;
+	bool bSuccess = false;
+
+	// Normalize packagenames and filenames
+	Algo::Transform(InFilenames, PackageNames, [](const FString& FileName)
+	{
+		return FPackageName::FilenameToLongPackageName(FileName);
+	});
+
+	// Remove packages if they are loaded actors or world
+	PackageNames.RemoveAll([&FilteredActorPackages, &LoadedPackages](const FString& PackageName) -> bool
+	{
+		UPackage* Package = FindPackage(NULL, *PackageName);
+		
+		if (Package != nullptr)
+		{
+			if (UWorld::IsWorldOrExternalActorPackage(Package))
+			{
+				FilteredActorPackages.Emplace(PackageName);
+				return true; // remove the package
+			}
+
+			LoadedPackages.Add(Package);
+		}
+
+		return false; // do not remove the package
+	});
+
+	if (!FilteredActorPackages.IsEmpty())
+	{
+		TStringBuilder<2048> Builder;
+		Builder.Join(FilteredActorPackages, TEXT(", "));
+		const FString Packages = Builder.ToString();
+
+		UE_LOG(LogSourceControl, Warning, TEXT("This operation could not complete on the following map or external packages, please unload them before retrying : %s"), *Packages);
+		return false;
+	}
+
+	// Prepare the packages to be reverted...
+	for (UPackage* Package : LoadedPackages)
+	{
+		// Detach the linkers of any loaded packages so that SCC can overwrite the files...
+		if (!Package->IsFullyLoaded())
+		{
+			FlushAsyncLoading();
+			Package->FullyLoad();
+		}
+		ResetLoaders(Package);
+	}
+
+	PackageFilenames = SourceControlHelpers::PackageFilenames(PackageNames);
+
+	// Apply Operation
+	bSuccess = InOperation(PackageFilenames);
+
+	// Operation may have deleted some packages, so we need to unload those rather than re-load them...
+	TArray<UPackage*> PackagesToUnload;
+	
+	LoadedPackages.RemoveAll([&](UPackage* InPackage) -> bool
+	{
+		const FString PackageExtension = InPackage->ContainsMap() ? FPackageName::GetMapPackageExtension() : FPackageName::GetAssetPackageExtension();
+		const FString PackageFilename = FPackageName::LongPackageNameToFilename(InPackage->GetName(), PackageExtension);
+		
+		if (!FPaths::FileExists(PackageFilename))
+		{
+			PackagesToUnload.Emplace(InPackage);
+			
+			return true; // remove package
+		}
+		return false; // keep package
+	});
+
+	// Hot-reload the new packages...
+	UPackageTools::ReloadPackages(LoadedPackages);
+
+	// Unload any deleted packages...
+	UPackageTools::UnloadPackages(PackagesToUnload);
+
+	// Re-cache the SCC state...
+	SourceControlProvider.Execute(ISourceControlOperation::Create<FUpdateStatus>(), PackageFilenames, EConcurrency::Asynchronous);
+
+	return bSuccess;
+}
+
+bool USourceControlHelpers::RevertAndReloadPackages(const TArray<FString>& InFilenames)
+{
+	auto RevertOperation = [](const TArray<FString>& InFilenames) -> bool
+	{
+		ISourceControlProvider& SourceControlProvider = ISourceControlModule::Get().GetProvider();
+
+		auto OperationCompleteCallback = FSourceControlOperationComplete::CreateLambda([](const FSourceControlOperationRef& Operation, ECommandResult::Type InResult)
+		{
+			if (Operation->GetName() == TEXT("Revert"))
+			{
+				TSharedRef<FRevert> RevertOperation = StaticCastSharedRef<FRevert>(Operation);
+				ISourceControlModule::Get().GetOnFilesDeleted().Broadcast(RevertOperation->GetDeletedFiles());
+			}
+		});
+
+		return SourceControlProvider.Execute(ISourceControlOperation::Create<FRevert>(), InFilenames, EConcurrency::Synchronous, OperationCompleteCallback) == ECommandResult::Succeeded;
+	};
+
+	return ApplyOperationAndReloadPackages(InFilenames, RevertOperation);
+}
+#endif //!WITH_EDITOR
 
 bool USourceControlHelpers::RevertFiles(const TArray<FString>& InFiles,	bool bSilent)
 {
diff --git a/Engine/Source/Developer/SourceControl/Public/SourceControlHelpers.h b/Engine/Source/Developer/SourceControl/Public/SourceControlHelpers.h
--- a/Engine/Source/Developer/SourceControl/Public/SourceControlHelpers.h
+++ b/Engine/Source/Developer/SourceControl/Public/SourceControlHelpers.h
@@ -347,6 +347,26 @@ public:
 	UFUNCTION(BlueprintCallable, Category = "Editor Scripting | Editor Source Control Helpers")
 	static bool RevertFile(const FString& InFile, bool bSilent = false);
 
+#if WITH_EDITOR
+
+	/**
+	 * Applies the provided function and reloads provided packages.
+	 * @param	InFilenames		The files/packages to apply the operation and reload
+	 * @param	InOperation		The function to apply
+	 * @return true if succeeded.
+	 */
+	static bool ApplyOperationAndReloadPackages(const TArray<FString>& InFilenames, const TFunctionRef<bool(const TArray<FString>&)>& InOperation);
+
+	/**
+     * Reverts the provided files then reload packages.
+     * @param	InFilenames		The files/packages to revert and reload
+     * @return true if succeeded.
+     */
+	UFUNCTION(BlueprintCallable, Category = "Editor Scripting | Editor Source Control Helpers")
+	static bool RevertAndReloadPackages(const TArray<FString>& InFilenames);
+
+#endif // !WITH_EDITOR
+
 	/**
 	 * Use currently set source control provider to revert files regardless whether any changes will be lost or not.
 	 * @note	Blocks until action is complete.
diff --git a/Engine/Source/Developer/UncontrolledChangelists/Private/UncontrolledChangelistsModule.cpp b/Engine/Source/Developer/UncontrolledChangelists/Private/UncontrolledChangelistsModule.cpp
--- a/Engine/Source/Developer/UncontrolledChangelists/Private/UncontrolledChangelistsModule.cpp
+++ b/Engine/Source/Developer/UncontrolledChangelists/Private/UncontrolledChangelistsModule.cpp
@@ -21,6 +21,7 @@
 #include "PackagesDialog.h"
 #include "Serialization/JsonSerializer.h"
 #include "Serialization/JsonWriter.h"
+#include "SourceControlHelpers.h"
 #include "SourceControlOperations.h"
 #include "Styling/SlateTypes.h"
 #include "UObject/ObjectSaveContext.h"
@@ -253,18 +254,24 @@ void FUncontrolledChangelistsModule::OnAssetAddedInternal(const FAssetData& Asse
 	}
 }
 
-void FUncontrolledChangelistsModule::OnRevert(const TArray<FString>& InFilenames)
+static bool ExecuteRevertOperation(const TArray<FString>& InFilenames)
 {
-	if (!IsEnabled() || InFilenames.IsEmpty())
-	{
-		return;
-	}
-
-	ISourceControlProvider& SourceControlProvider = ISourceControlModule::Get().GetProvider();
-	IFileManager& FileManager = IFileManager::Get();
+	ISourceControlModule& SourceControlModule = ISourceControlModule::Get();
+	ISourceControlProvider& SourceControlProvider = SourceControlModule.GetProvider();
 	TArray<FSourceControlStateRef> UpdatedFilestates;
 
-	SourceControlProvider.GetState(InFilenames, UpdatedFilestates, EStateCacheUsage::ForceUpdate);
+	auto BuildFileString = [](const TArray<FString>& Files) -> FString
+	{
+		TStringBuilder<2048> Builder;
+		Builder.Join(Files, TEXT(", "));
+		return Builder.ToString();
+	};
+
+	if (SourceControlProvider.GetState(InFilenames, UpdatedFilestates, EStateCacheUsage::ForceUpdate) != ECommandResult::Succeeded)
+	{
+		UE_LOG(LogSourceControl, Error, TEXT("Failed to update the source control files states for %s."), *BuildFileString(InFilenames));
+		return false;
+	}
 
 	TArray<FString> FilesToDelete;
 	TArray<FString> FilesToRevert;
@@ -286,19 +293,49 @@ void FUncontrolledChangelistsModule::OnRevert(const TArray<FString>& InFilenames
 		TSharedRef<FSync> ForceSyncOperation = ISourceControlOperation::Create<FSync>();
 		ForceSyncOperation->SetForce(true);
 		ForceSyncOperation->SetLastSyncedFlag(true);
-		SourceControlProvider.Execute(ForceSyncOperation, FilesToRevert);
+
+		if (SourceControlProvider.Execute(ForceSyncOperation, FilesToRevert) != ECommandResult::Succeeded)
+		{
+			UE_LOG(LogSourceControl, Error, TEXT("Failed to sync the following files to a previous version: %s."), *BuildFileString(FilesToRevert));
+			return false;
+		}
 	}
 
+	IFileManager& FileManager = IFileManager::Get();
+	bool bSuccess = true;
+
 	for (const FString& FileToDelete : FilesToDelete)
 	{
 		const bool bRequireExists = true;
 		const bool bEvenReadOnly = false;
 		const bool bQuiet = false;
 
-		FileManager.Delete(*FileToDelete, bRequireExists, bEvenReadOnly, bQuiet);
+		if (!FileManager.Delete(*FileToDelete, bRequireExists, bEvenReadOnly, bQuiet))
+		{
+			UE_LOG(LogSourceControl, Error, TEXT("Failed to delete %s."), *FileToDelete);
+			bSuccess = false;
+		}
 	}
 
+	SourceControlModule.GetOnFilesDeleted().Broadcast(FilesToDelete);
+
+	return bSuccess;
+}
+
+bool FUncontrolledChangelistsModule::OnRevert(const TArray<FString>& InFilenames)
+{
+	bool bSuccess = false;
+
+	if (!IsEnabled() || InFilenames.IsEmpty())
+	{
+		return true;
+	}
+
+	bSuccess = SourceControlHelpers::ApplyOperationAndReloadPackages(InFilenames, ExecuteRevertOperation);
+
 	UpdateStatus();
+
+	return bSuccess;
 }
 
 void FUncontrolledChangelistsModule::OnObjectPreSaved(UObject* InObject, const FObjectPreSaveContext& InPreSaveContext)
diff --git a/Engine/Source/Developer/UncontrolledChangelists/Public/UncontrolledChangelistsModule.h b/Engine/Source/Developer/UncontrolledChangelists/Public/UncontrolledChangelistsModule.h
--- a/Engine/Source/Developer/UncontrolledChangelists/Public/UncontrolledChangelistsModule.h
+++ b/Engine/Source/Developer/UncontrolledChangelists/Public/UncontrolledChangelistsModule.h
@@ -91,9 +91,10 @@ public:
 
 	/** Called when "Revert files" button is clicked. Reverts modified files and deletes new ones.
 	 *  @param	InFilenames		The files to be reverted
+	 *	@return true if the provided files were reverted.
 	 */
-	void OnRevert(const TArray<FString>& InFilenames);
-
+	bool OnRevert(const TArray<FString>& InFilenames);
+	
 	/**
 	 * Delegate callback called before an asset has been written to disk.
 	 * @param 	InObject 			The saved object.
diff --git a/Engine/Source/Editor/SourceControlWindows/Private/SSourceControlChangelists.cpp b/Engine/Source/Editor/SourceControlWindows/Private/SSourceControlChangelists.cpp
--- a/Engine/Source/Editor/SourceControlWindows/Private/SSourceControlChangelists.cpp
+++ b/Engine/Source/Editor/SourceControlWindows/Private/SSourceControlChangelists.cpp
@@ -750,17 +750,28 @@ void SSourceControlChangelistsWidget::OnRevert()
 	ISourceControlProvider& SourceControlProvider = ISourceControlModule::Get().GetProvider();
 
 	// Reverts the selected Changelist or Files
-	if (SelectedChangelist.IsValid() || (!SelectedControlledFiles.IsEmpty()))
+	if (SelectedChangelist.IsValid())
 	{
-		auto RevertOperation = ISourceControlOperation::Create<FRevert>();
-		SourceControlProvider.Execute(RevertOperation, SelectedChangelist, SelectedControlledFiles, EConcurrency::Synchronous, FSourceControlOperationComplete::CreateLambda([](const FSourceControlOperationRef& Operation, ECommandResult::Type InResult)
+		TSet<FString> ToRevertFiles;
+		FSourceControlChangelistStatePtr ChangelistState;
+		ChangelistState = SourceControlProvider.GetState(SelectedChangelist.ToSharedRef(), EStateCacheUsage::Use);
+
+		if (ChangelistState.IsValid())
 		{
-			if (Operation->GetName() == TEXT("Revert"))
+			Algo::Transform(ChangelistState->GetFilesStates(), ToRevertFiles, [](FSourceControlStateRef FileState)
 			{
-				TSharedRef<FRevert> RevertOperation = StaticCastSharedRef<FRevert>(Operation);
-				ISourceControlModule::Get().GetOnFilesDeleted().Broadcast(RevertOperation->GetDeletedFiles());
-			}
-		}));
+				return FileState->GetFilename();
+			});
+		}
+
+		// Enforce unique occurences of files.
+		ToRevertFiles.Append(SelectedControlledFiles);
+		SelectedControlledFiles = ToRevertFiles.Array();
+	}
+
+	if (!SelectedControlledFiles.IsEmpty())
+	{
+		SourceControlHelpers::RevertAndReloadPackages(SelectedControlledFiles);
 	}
 
 	FUncontrolledChangelistStatePtr SelectedUncontrolledChangelist = GetCurrentUncontrolledChangelistState();
diff --git a/Engine/Source/Editor/SourceControlWindows/Private/SSourceControlRevert.cpp b/Engine/Source/Editor/SourceControlWindows/Private/SSourceControlRevert.cpp
--- a/Engine/Source/Editor/SourceControlWindows/Private/SSourceControlRevert.cpp
+++ b/Engine/Source/Editor/SourceControlWindows/Private/SSourceControlRevert.cpp
@@ -437,58 +437,10 @@ bool FSourceControlWindows::PromptForRevert( const TArray<FString>& InPackageNam
 		{
 			TArray<FString> FinalPackagesToRevert;
 			SourceControlWidget->GetPackagesToRevert(FinalPackagesToRevert);
+			
 			if (FinalPackagesToRevert.Num() > 0)
 			{
-				// attempt to unload the packages we are about to revert
-				TArray<UPackage*> LoadedPackages;
-				for (TArray<FString>::TConstIterator PackageIter(InPackageNames); PackageIter; ++PackageIter)
-				{
-					UPackage* Package = FindPackage(NULL, **PackageIter);
-					if (Package != NULL)
-					{
-						LoadedPackages.Add(Package);
-					}
-				}
-
-				const TArray<FString> RevertPackageFilenames = SourceControlHelpers::PackageFilenames(FinalPackagesToRevert);
-
-				// Prepare the packages to be reverted...
-				for (UPackage* Package : LoadedPackages)
-				{
-					// Detach the linkers of any loaded packages so that SCC can overwrite the files...
-					if (!Package->IsFullyLoaded())
-					{
-						FlushAsyncLoading();
-						Package->FullyLoad();
-					}
-					ResetLoaders(Package);
-				}
-
-				// Revert everything...
-				SourceControlProvider.Execute(ISourceControlOperation::Create<FRevert>(), RevertPackageFilenames);
-
-				// Reverting may have deleted some packages, so we need to unload those rather than re-load them...
-				TArray<UPackage*> PackagesToUnload;
-				LoadedPackages.RemoveAll([&](UPackage* InPackage) -> bool
-				{
-					const FString PackageExtension = InPackage->ContainsMap() ? FPackageName::GetMapPackageExtension() : FPackageName::GetAssetPackageExtension();
-					const FString PackageFilename = FPackageName::LongPackageNameToFilename(InPackage->GetName(), PackageExtension);
-					if (!FPaths::FileExists(PackageFilename))
-					{
-						PackagesToUnload.Emplace(InPackage);
-						return true; // remove package
-					}
-					return false; // keep package
-				});
-
-				// Hot-reload the new packages...
-				UPackageTools::ReloadPackages(LoadedPackages);
-
-				// Unload any deleted packages...
-				UPackageTools::UnloadPackages(PackagesToUnload);
-
-				// Re-cache the SCC state...
-				SourceControlProvider.Execute(ISourceControlOperation::Create<FUpdateStatus>(), RevertPackageFilenames, EConcurrency::Asynchronous);
+				SourceControlHelpers::RevertAndReloadPackages(FinalPackagesToRevert);
 
 				bReverted = true;
 			}
