Fixes packages from different project paths unable to be reverted.

#changelist validated

diff --git a/Engine/Source/Developer/SourceControl/Private/SourceControlHelpers.cpp b/Engine/Source/Developer/SourceControl/Private/SourceControlHelpers.cpp
--- a/Engine/Source/Developer/SourceControl/Private/SourceControlHelpers.cpp
+++ b/Engine/Source/Developer/SourceControl/Private/SourceControlHelpers.cpp
@@ -912,10 +912,19 @@ bool USourceControlHelpers::ApplyOperationAndReloadPackages(const TArray<FString
 	bool bSuccess = false;
 
 	// Normalize packagenames and filenames
-	Algo::Transform(InFilenames, PackageNames, [](const FString& FileName)
+	for (const FString& Filename : InFilenames)
 	{
-		return FPackageName::FilenameToLongPackageName(FileName);
-	});
+		FString Result;
+		
+		if (FPackageName::TryConvertFilenameToLongPackageName(Filename, Result))
+		{
+			PackageNames.Emplace(MoveTemp(Result));
+		}
+		else
+		{
+			PackageNames.Emplace(Filename);
+		}
+	}
 
 	// Remove packages if they are loaded actors or world
 	PackageNames.RemoveAll([&FilteredActorPackages, &LoadedPackages](const FString& PackageName) -> bool
