Fixed memory leak when a long source control operation blocks the main thread while the rendering thread renders to animate the source control feedback popup.

Note: This popup is currently disabled for P4.

diff --git a/Engine/Source/Developer/SourceControl/Private/ScopedSourceControlProgress.cpp b/Engine/Source/Developer/SourceControl/Private/ScopedSourceControlProgress.cpp
--- a/Engine/Source/Developer/SourceControl/Private/ScopedSourceControlProgress.cpp
+++ b/Engine/Source/Developer/SourceControl/Private/ScopedSourceControlProgress.cpp
@@ -18,6 +18,8 @@
 #include "Widgets/Images/SImage.h"
 #include "Framework/Docking/TabManager.h"
 #include "Framework/Application/SlateApplication.h"
+#include "RHI.h"
+#include "RenderingThread.h"
 
 #define LOCTEXT_NAMESPACE "SourceControlProgress"
 
@@ -208,7 +210,7 @@ FScopedSourceControlProgress::FScopedSourceControlProgress(const FText& InText,
 		TSharedPtr<SWindow> RootWindow = FGlobalTabmanager::Get()->GetRootWindow();
 		FSlateApplication::Get().AddModalWindow(Window, RootWindow, true);
 
-		Window->ShowWindow();	
+		Window->ShowWindow();
 	
 		Tick();
 	}
@@ -226,10 +228,22 @@ void FScopedSourceControlProgress::Tick()
 {
 	if (!(FApp::IsUnattended() || IsRunningCommandlet()) && WindowPtr.IsValid() && FSlateApplication::Get().CanDisplayWindows())
 	{
+		// Mark begin frame
+		if (GIsRHIInitialized)
+		{
+			ENQUEUE_RENDER_COMMAND(BeginFrameCmd)([](FRHICommandListImmediate& RHICmdList) { RHICmdList.BeginFrame(); });
+		}
+
 		// Tick Slate application
 		FSlateApplication::Get().Tick();
 
-		// Sync the game thread and the render thread
+		// End frame so frame fence number gets incremented
+		if (GIsRHIInitialized)
+		{
+			ENQUEUE_RENDER_COMMAND(EndFrameCmd)([](FRHICommandListImmediate& RHICmdList) { RHICmdList.EndFrame(); });
+		}
+
+		// Sync the game thread and the render thread. This is needed if many StatusUpdate are called.
 		FSlateApplication::Get().GetRenderer()->Sync();
 	}
 }
diff --git a/Engine/Source/Developer/SourceControl/SourceControl.Build.cs b/Engine/Source/Developer/SourceControl/SourceControl.Build.cs
--- a/Engine/Source/Developer/SourceControl/SourceControl.Build.cs
+++ b/Engine/Source/Developer/SourceControl/SourceControl.Build.cs
@@ -24,6 +24,8 @@ public class SourceControl : ModuleRules
 			new string[] {
 				"AssetRegistry",
 				"Slate",
+				"RenderCore",
+				"RHI",
 			}
 		);
 
