Enabled Uncontrolled Changelists by default.
Adds setting to disable uncontrolled changelists.

#changelist validated

diff --git a/Engine/Source/Developer/SourceControl/Private/SourceControlPreferences.cpp b/Engine/Source/Developer/SourceControl/Private/SourceControlPreferences.cpp
--- a/Engine/Source/Developer/SourceControl/Private/SourceControlPreferences.cpp
+++ b/Engine/Source/Developer/SourceControl/Private/SourceControlPreferences.cpp
@@ -11,3 +11,8 @@ bool USourceControlPreferences::ShouldDeleteNewFilesOnRevert()
 {
 	return GetDefault<USourceControlPreferences>()->bShouldDeleteNewFilesOnRevert;
 }
+
+bool USourceControlPreferences::AreUncontrolledChangelistsEnabled()
+{
+	return GetDefault<USourceControlPreferences>()->bEnableUncontrolledChangelists;
+}
diff --git a/Engine/Source/Developer/SourceControl/Public/SourceControlPreferences.h b/Engine/Source/Developer/SourceControl/Public/SourceControlPreferences.h
--- a/Engine/Source/Developer/SourceControl/Public/SourceControlPreferences.h
+++ b/Engine/Source/Developer/SourceControl/Public/SourceControlPreferences.h
@@ -29,6 +29,9 @@ public:
 	/** Helper to access easily Should Delete New Files On Revert setting */
 	static bool ShouldDeleteNewFilesOnRevert();
 
+	/** Helper to access easily Enable Uncontrolled Changelists setting */
+	static bool AreUncontrolledChangelistsEnabled();
+
 public:
 	/** If enabled, adds a tag in changelist descriptions when they are validated */
 	UPROPERTY(config, EditAnywhere, Category = SourceControl, meta = (ToolTip = "Adds validation tag to changelist description on submit."))
@@ -38,6 +41,10 @@ public:
 	UPROPERTY(config, EditAnywhere, Category = SourceControl, meta = (ToolTip = "Deletes new files when reverted."))
 	bool bShouldDeleteNewFilesOnRevert = true;
 
+	/** Enables Uncontrolled Changelists features. */
+	UPROPERTY(config, EditAnywhere, Category = SourceControl, meta = (ToolTip = "Enables Uncontrolled Changelists features. The editor must be restarted for the change to be fully taken into account.", ConfigRestartRequired = true))
+	bool bEnableUncontrolledChangelists = true;
+
 	/** List of lines to add to any collection on checkin */
 	UPROPERTY(config, EditAnywhere, Category = SourceControl)
 	TArray<FString> CollectionChangelistTags;
diff --git a/Engine/Source/Developer/UncontrolledChangelists/Private/UncontrolledChangelistsModule.cpp b/Engine/Source/Developer/UncontrolledChangelists/Private/UncontrolledChangelistsModule.cpp
--- a/Engine/Source/Developer/UncontrolledChangelists/Private/UncontrolledChangelistsModule.cpp
+++ b/Engine/Source/Developer/UncontrolledChangelists/Private/UncontrolledChangelistsModule.cpp
@@ -23,17 +23,12 @@
 #include "Serialization/JsonWriter.h"
 #include "SourceControlHelpers.h"
 #include "SourceControlOperations.h"
+#include "SourceControlPreferences.h"
 #include "Styling/SlateTypes.h"
 #include "UObject/ObjectSaveContext.h"
 
 #define LOCTEXT_NAMESPACE "UncontrolledChangelists"
 
-static TAutoConsoleVariable<bool> CVarUncontrolledChangelistsEnable(
-	TEXT("UncontrolledChangelists.Enable"),
-	false,
-	TEXT("Enables Uncontrolled Changelists (experimental).")
-);
-
 void FUncontrolledChangelistsModule::FStartupTask::DoWork()
 {
 	double StartTime = FPlatformTime::Seconds();
@@ -59,7 +54,8 @@ void FUncontrolledChangelistsModule::FStartupTask::DoWork()
 
 void FUncontrolledChangelistsModule::StartupModule()
 {
-	bIsEnabled = CVarUncontrolledChangelistsEnable.GetValueOnGameThread();
+	bIsEnabled = USourceControlPreferences::AreUncontrolledChangelistsEnabled();
+	
 	if (!IsEnabled())
 	{
 		return;
