Fixed a crash that happening when a an asset that was just added to source control was reverted and the package was unloaded.

diff --git a/Engine/Source/Developer/SourceControl/Private/SourceControlHelpers.cpp b/Engine/Source/Developer/SourceControl/Private/SourceControlHelpers.cpp
--- a/Engine/Source/Developer/SourceControl/Private/SourceControlHelpers.cpp
+++ b/Engine/Source/Developer/SourceControl/Private/SourceControlHelpers.cpp
@@ -905,7 +905,7 @@ bool USourceControlHelpers::RevertFile(const FString& InFile, bool bSilent)
 bool USourceControlHelpers::ApplyOperationAndReloadPackages(const TArray<FString>& InFilenames, const TFunctionRef<bool(const TArray<FString>&)>& InOperation)
 {
 	ISourceControlProvider& SourceControlProvider = ISourceControlModule::Get().GetProvider();
-	TArray<UPackage*> LoadedPackages;
+	TArray<TWeakObjectPtr<UPackage>> LoadedPackages;
 	TArray<FString> PackageNames;
 	TArray<FString> PackageFilenames;
 	TArray<FString> FilteredActorPackages;
@@ -956,7 +956,7 @@ bool USourceControlHelpers::ApplyOperationAndReloadPackages(const TArray<FString
 	}
 
 	// Prepare the packages to be reverted...
-	for (UPackage* Package : LoadedPackages)
+	for (TWeakObjectPtr<UPackage>& Package : LoadedPackages)
 	{
 		// Detach the linkers of any loaded packages so that SCC can overwrite the files...
 		if (!Package->IsFullyLoaded())
@@ -964,7 +964,7 @@ bool USourceControlHelpers::ApplyOperationAndReloadPackages(const TArray<FString
 			FlushAsyncLoading();
 			Package->FullyLoad();
 		}
-		ResetLoaders(Package);
+		ResetLoaders(Package.Get());
 	}
 
 	PackageFilenames = SourceControlHelpers::PackageFilenames(PackageNames);
@@ -974,23 +974,26 @@ bool USourceControlHelpers::ApplyOperationAndReloadPackages(const TArray<FString
 
 	// Operation may have deleted some packages, so we need to unload those rather than re-load them...
 	TArray<UPackage*> PackagesToUnload;
-	
-	LoadedPackages.RemoveAll([&](UPackage* InPackage) -> bool
+	TArray<UPackage*> PackagesToReload;
+	for (TWeakObjectPtr<UPackage>& Package : LoadedPackages)
 	{
-		const FString PackageExtension = InPackage->ContainsMap() ? FPackageName::GetMapPackageExtension() : FPackageName::GetAssetPackageExtension();
-		const FString PackageFilename = FPackageName::LongPackageNameToFilename(InPackage->GetName(), PackageExtension);
-		
+		if (!Package.IsValid()) // Package was already unloaded on delete?
+		{
+			continue;
+		}
+
+		const FString PackageExtension = Package->ContainsMap() ? FPackageName::GetMapPackageExtension() : FPackageName::GetAssetPackageExtension();
+		const FString PackageFilename = FPackageName::LongPackageNameToFilename(Package->GetName(), PackageExtension);
 		if (!FPaths::FileExists(PackageFilename))
 		{
-			PackagesToUnload.Emplace(InPackage);
-			
-			return true; // remove package
+			PackagesToUnload.Emplace(Package.Get());
 		}
-		return false; // keep package
-	});
+
+		PackagesToReload.Emplace(Package.Get());
+	};
 
 	// Hot-reload the new packages...
-	UPackageTools::ReloadPackages(LoadedPackages);
+	UPackageTools::ReloadPackages(PackagesToReload);
 
 	// Unload any deleted packages...
 	UPackageTools::UnloadPackages(PackagesToUnload);
