The source control settings for the VA source control connection are no longer saved to a local ini file.

#rnx

### Virtualization
- At the moment there is no way in the editor UX to change the source control settings for the source control backend, so what ever settings are found in the global environment would be applied when the editor is first run and then saved to a config file under the <saved> directory. From this point onwards those settings would always be used when setting up the source control backend. If the settings were wrong, the first time that the editor was started, then we'd continue to use the incorrect settings even if the global enviroment was fixed. Making the backend work at that point would require that the user know about the local ini file and change the settings there.
- We cannot fix this by just ignoring the ini file as it has been requested that users can add their settings there in case they want to customize things, so we just need to avoid saving to it.
- Added documentation to the source control backends header file on how to set up the ini file if needed.

### Perforce
- FSourceControlInitSettings now accepts EConfigBehavior which describes how we want the source control provider to deal with its settings with regards to the ini file.
- Note that EConfigBehavior is not a bitfield as we do not wish to support only writing to the ini file, there is no point if it cannot be read from.
- When the source control provider is creatred we pass on the into from EConfigBehavior to the FPerforceSourceControlSettings to enable/disable saving and loading.

diff --git a/Engine/Plugins/Developer/PerforceSourceControl/Source/PerforceSourceControl/Private/PerforceSourceControlProvider.cpp b/Engine/Plugins/Developer/PerforceSourceControl/Source/PerforceSourceControl/Private/PerforceSourceControlProvider.cpp
--- a/Engine/Plugins/Developer/PerforceSourceControl/Source/PerforceSourceControl/Private/PerforceSourceControlProvider.cpp
+++ b/Engine/Plugins/Developer/PerforceSourceControl/Source/PerforceSourceControl/Private/PerforceSourceControlProvider.cpp
@@ -60,6 +60,9 @@ FPerforceSourceControlProvider::FPerforceSourceControlProvider(const FStringView
 	, bLoginError(false)
 	, PersistentConnection(nullptr)
 {
+	AccessSettings().SetAllowSave(InInitialSettings.CanWriteToConfigFile());
+	AccessSettings().SetAllowLoad(InInitialSettings.CanReadFromConfigFile());
+
 	AccessSettings().LoadSettings();
 }
 
diff --git a/Engine/Plugins/Developer/PerforceSourceControl/Source/PerforceSourceControl/Private/PerforceSourceControlSettings.cpp b/Engine/Plugins/Developer/PerforceSourceControl/Source/PerforceSourceControl/Private/PerforceSourceControlSettings.cpp
--- a/Engine/Plugins/Developer/PerforceSourceControl/Source/PerforceSourceControl/Private/PerforceSourceControlSettings.cpp
+++ b/Engine/Plugins/Developer/PerforceSourceControl/Source/PerforceSourceControl/Private/PerforceSourceControlSettings.cpp
@@ -12,6 +12,8 @@
 
 FPerforceSourceControlSettings::FPerforceSourceControlSettings(const FPerforceSourceControlProvider& InSCCProvider, const FStringView& OwnerName)
 	: SCCProvider(InSCCProvider)
+	, bCanSaveToIniFile(true)
+	, bCanLoadFromIniFile(true)
 {
 	if (OwnerName.IsEmpty())
 	{
@@ -107,8 +109,23 @@ void FPerforceSourceControlSettings::ImportP4Config()
 	ConnectionInfo.Workspace = ANSI_TO_TCHAR(TestP4.GetClient().Text());
 }
 
+void FPerforceSourceControlSettings::SetAllowSave(bool bFlag)
+{
+	bCanSaveToIniFile = bFlag;
+}
+
+void FPerforceSourceControlSettings::SetAllowLoad(bool bFlag)
+{
+	bCanLoadFromIniFile = bFlag;
+}
+
 void FPerforceSourceControlSettings::LoadSettings()
 {
+	if (!bCanLoadFromIniFile)
+	{
+		return;
+	}
+
 	FScopeLock ScopeLock(&CriticalSection);
 	const FString& IniFile = SourceControlHelpers::GetSettingsIni();
 
@@ -132,7 +149,7 @@ void FPerforceSourceControlSettings::LoadSettings()
 
 void FPerforceSourceControlSettings::SaveSettings() const
 {
-	if (FApp::IsUnattended() || IsRunningCommandlet())
+	if (FApp::IsUnattended() || IsRunningCommandlet() || !bCanSaveToIniFile)
 	{
 		return;
 	}
diff --git a/Engine/Plugins/Developer/PerforceSourceControl/Source/PerforceSourceControl/Private/PerforceSourceControlSettings.h b/Engine/Plugins/Developer/PerforceSourceControl/Source/PerforceSourceControl/Private/PerforceSourceControlSettings.h
--- a/Engine/Plugins/Developer/PerforceSourceControl/Source/PerforceSourceControl/Private/PerforceSourceControlSettings.h
+++ b/Engine/Plugins/Developer/PerforceSourceControl/Source/PerforceSourceControl/Private/PerforceSourceControlSettings.h
@@ -47,6 +47,12 @@ public:
 
 	void SetUseP4Config(bool bInUseP4Config);
 
+	/** When set to true the settings can be saved to an ini file, when false nothing will be saved */
+	void SetAllowSave(bool bFlag);
+	
+	/** When set to true the settings can be read from an ini file, when false nothing will be read */
+	void SetAllowLoad(bool bFlag);
+
 	/** Load settings from ini file */
 	void LoadSettings();
 
@@ -77,4 +83,10 @@ private:
 
 	/** The source control provider that this object is associated with */
 	const FPerforceSourceControlProvider& SCCProvider;
+
+	/** Controls whether the settings can be saved to an ini file or not */
+	bool bCanSaveToIniFile;
+
+	/** Controls whether the settings can be read from an ini file or not */
+	bool bCanLoadFromIniFile;
 };
diff --git a/Engine/Source/Developer/SourceControl/Private/SourceControlInitSettings.cpp b/Engine/Source/Developer/SourceControl/Private/SourceControlInitSettings.cpp
--- a/Engine/Source/Developer/SourceControl/Private/SourceControlInitSettings.cpp
+++ b/Engine/Source/Developer/SourceControl/Private/SourceControlInitSettings.cpp
@@ -5,16 +5,28 @@
 #include "Containers/StringView.h"
 
 FSourceControlInitSettings::FSourceControlInitSettings(EBehavior InBehavior)
-	: Behavior(InBehavior)
+	: OverrideBehavior(InBehavior)
 {
 
 }
 
-void FSourceControlInitSettings::AddSetting(FStringView SettingName, FStringView SettingValue)
+void FSourceControlInitSettings::SetConfigBehavior(EConfigBehavior InBehavior)
+{
+	ConfigBehavior = InBehavior;
+}
+
+bool FSourceControlInitSettings::CanWriteToConfigFile() const
+{
+	return ConfigBehavior == EConfigBehavior::ReadWrite;
+}
+
+bool FSourceControlInitSettings::CanReadFromConfigFile() const
 {
-//	const int32 Hash = GetTypeHash(SettingName);
-//	Settings.AddByHash(Hash, SettingName, SettingValue)
+	return ConfigBehavior != EConfigBehavior::None;
+}
 
+void FSourceControlInitSettings::AddSetting(FStringView SettingName, FStringView SettingValue)
+{
 	Settings.Add(FString(SettingName), FString(SettingValue));
 }
 
@@ -27,7 +39,7 @@ void FSourceControlInitSettings::OverrideSetting(FStringView SettingName, FStrin
 	{
 		InOutSettingValue = *InitialValue;
 	}
-	else if (Behavior == EBehavior::OverrideAll)
+	else if (OverrideBehavior == EBehavior::OverrideAll)
 	{
 		InOutSettingValue.Empty();
 	}
diff --git a/Engine/Source/Developer/SourceControl/Public/SourceControlInitSettings.h b/Engine/Source/Developer/SourceControl/Public/SourceControlInitSettings.h
--- a/Engine/Source/Developer/SourceControl/Public/SourceControlInitSettings.h
+++ b/Engine/Source/Developer/SourceControl/Public/SourceControlInitSettings.h
@@ -15,20 +15,38 @@ public:
 	{
 		/** All existing settings will be overridden via the contents of FSourceControlInitSettings. Settings that are not found will be reset to default states */
 		OverrideAll,
-		/** Only the settings found in FSourceControlInitSettings will be overriden. Settings not found will be left with their current values. */
+		/** Only the settings found in FSourceControlInitSettings will be overridden. Settings not found will be left with their current values. */
 		OverrideExisting,
 	};
 
+	enum class EConfigBehavior
+	{
+		/** Can both read from, and save to the ini file*/
+		ReadWrite,
+		/** Will only read settings from the ini file, settings determined at runtime will not be saved to the ini file */
+		ReadOnly,
+		/** The settings will not be saved to the ini file, nor will they be read from the ini file  */
+		None
+	};
+
 	FSourceControlInitSettings(EBehavior Behavior);
 	~FSourceControlInitSettings() = default;
 
+	void SetConfigBehavior(EConfigBehavior Behavior);
+
+	bool CanWriteToConfigFile() const;
+	bool CanReadFromConfigFile() const;
+		
 	void AddSetting(FStringView SettingName, FStringView SettingValue);
 	void OverrideSetting(FStringView SettingName, FString& InOutSettingValue);
 
 	bool HasOverrides() const;
 	bool IsOverridden(FStringView SettingName) const;
+
 private:
-	EBehavior Behavior;
+
+	EBehavior OverrideBehavior = EBehavior::OverrideAll;
+	EConfigBehavior ConfigBehavior = EConfigBehavior::ReadWrite;
 
 	TMap<FString, FString> Settings;
 };
diff --git a/Engine/Source/Developer/Virtualization/Private/VirtualizationSourceControlBackend.cpp b/Engine/Source/Developer/Virtualization/Private/VirtualizationSourceControlBackend.cpp
--- a/Engine/Source/Developer/Virtualization/Private/VirtualizationSourceControlBackend.cpp
+++ b/Engine/Source/Developer/Virtualization/Private/VirtualizationSourceControlBackend.cpp
@@ -261,6 +261,8 @@ IVirtualizationBackend::EConnectionStatus FSourceControlBackend::OnConnect()
 
 	// We do not want the connection to have a client workspace so explicitly set it to empty
 	FSourceControlInitSettings SCCSettings(FSourceControlInitSettings::EBehavior::OverrideExisting);
+	SCCSettings.SetConfigBehavior(FSourceControlInitSettings::EConfigBehavior::ReadOnly);
+
 	if (!ServerAddress.IsEmpty())
 	{
 		SCCSettings.AddSetting(TEXT("P4Port"), ServerAddress);
diff --git a/Engine/Source/Developer/Virtualization/Private/VirtualizationSourceControlBackend.h b/Engine/Source/Developer/Virtualization/Private/VirtualizationSourceControlBackend.h
--- a/Engine/Source/Developer/Virtualization/Private/VirtualizationSourceControlBackend.h
+++ b/Engine/Source/Developer/Virtualization/Private/VirtualizationSourceControlBackend.h
@@ -56,6 +56,18 @@ class FSemaphore;
  *											should use as the root location to submit payloads from.
  *											If the users machine has this set then 'SubmitFromTempDir' 
  *											will be ignored. 
+ * 
+ * Perforce Setup
+ * The backend will attempt to use the current global environment settings for perforce to determine 
+ * the P4PORT and P4USER values. If you wish to set them via an ini file then they can be done per
+ * user by adding the following to Saved\Config\<platform>Editor\SourceControlSettings.ini
+ * 
+ * -----
+ * [PerforceSourceControl.VirtualizationSettings]
+ * Port=<server address>
+ * UserName=<username>
+ * -----
+ * Note that the backend will read from these ini file settings but will not write to them.
  */
 class FSourceControlBackend final : public IVirtualizationBackend
 {
