DefaultSourceControlProvider: Fix missing Make Writable option when disconnected from SCC (introduced in CL 22187137)

diff --git a/Engine/Source/Developer/SourceControl/Private/DefaultSourceControlProvider.cpp b/Engine/Source/Developer/SourceControl/Private/DefaultSourceControlProvider.cpp
--- a/Engine/Source/Developer/SourceControl/Private/DefaultSourceControlProvider.cpp
+++ b/Engine/Source/Developer/SourceControl/Private/DefaultSourceControlProvider.cpp
@@ -81,7 +81,7 @@ void FDefaultSourceControlProvider::CancelOperation( const FSourceControlOperati
 
 bool FDefaultSourceControlProvider::UsesLocalReadOnlyState() const
 {
-	return false;
+	return true;
 }
 
 bool FDefaultSourceControlProvider::UsesChangelists() const
