Fixed "SourceControl: Error: CommandMessage Command: UpdateStatus, Error: Usage: filelog" message being displayed when using GetFileHistory

diff --git a/Engine/Source/Developer/SourceControl/Private/SourceControlAssetDataCache.cpp b/Engine/Source/Developer/SourceControl/Private/SourceControlAssetDataCache.cpp
--- a/Engine/Source/Developer/SourceControl/Private/SourceControlAssetDataCache.cpp
+++ b/Engine/Source/Developer/SourceControl/Private/SourceControlAssetDataCache.cpp
@@ -139,9 +139,9 @@ void FSourceControlAssetDataCache::GetFileHistory()
 	SourceControlProvider.Execute(UpdateStatusOperation,
 								  FileHistoryToUpdate,
 								  EConcurrency::Asynchronous,
-								  FSourceControlOperationComplete::CreateLambda([FileHistoryToUpdate = MoveTemp(FileHistoryToUpdate)](const FSourceControlOperationRef& Operation, ECommandResult::Type InResult)
+								  FSourceControlOperationComplete::CreateLambda([FileHistoryToUpdateCopy = FileHistoryToUpdate](const FSourceControlOperationRef& Operation, ECommandResult::Type InResult)
 	{
-		ISourceControlModule::Get().GetAssetDataCache().OnUpdateHistoryComplete(FileHistoryToUpdate, Operation, InResult);
+		ISourceControlModule::Get().GetAssetDataCache().OnUpdateHistoryComplete(FileHistoryToUpdateCopy, Operation, InResult);
 	}));
 
 	FileHistoryToUpdate.Reset();
