Fix not being able to revert Uncontrolled Changelist changes with files not at head revision. (files were forced synched to proper revision but would stay in the UC)

- Add flag to SourceControlOperation to test modified status against synched revision (p4 #have)
- Use this new flag when testing for modifications in UncontrolledChangelist

diff --git a/Engine/Plugins/Developer/PerforceSourceControl/Source/PerforceSourceControl/Private/PerforceSourceControlOperations.cpp b/Engine/Plugins/Developer/PerforceSourceControl/Source/PerforceSourceControl/Private/PerforceSourceControlOperations.cpp
--- a/Engine/Plugins/Developer/PerforceSourceControl/Source/PerforceSourceControl/Private/PerforceSourceControlOperations.cpp
+++ b/Engine/Plugins/Developer/PerforceSourceControl/Source/PerforceSourceControl/Private/PerforceSourceControlOperations.cpp
@@ -1884,6 +1884,11 @@ bool FPerforceUpdateStatusWorker::Execute(FPerforceSourceControlCommand& InComma
 					// If the file is a directory, do a recursive diff on the contents
 					File /= TEXT("...");
 				}
+								
+				if (Operation->ShouldUpdateModifiedStateToLocalRevision())
+				{
+					File += TEXT("#have");
+				}
 
 				Parameters.Add(MoveTemp(File));
 			}
diff --git a/Engine/Source/Developer/SourceControl/Public/SourceControlOperations.h b/Engine/Source/Developer/SourceControl/Public/SourceControlOperations.h
--- a/Engine/Source/Developer/SourceControl/Public/SourceControlOperations.h
+++ b/Engine/Source/Developer/SourceControl/Public/SourceControlOperations.h
@@ -354,6 +354,7 @@ public:
 		: bUpdateHistory(false)
 		, bGetOpenedOnly(false)
 		, bUpdateModifiedState(false)
+		, bUpdateModifiedStateToLocalRevision(false)
 		, bCheckingAllFiles(false)
 		, bForceQuiet(false)
 		, bForceUpdate(false)
@@ -387,6 +388,11 @@ public:
 		bUpdateModifiedState = bInUpdateModifiedState;
 	}
 
+	void SetUpdateModifiedStateToLocalRevision(bool bInUpdateModifiedStateToLocalRevision)
+	{
+		bUpdateModifiedStateToLocalRevision = bInUpdateModifiedStateToLocalRevision;
+	}
+
 	void SetCheckingAllFiles( bool bInCheckingAllFiles )
 	{
 		bCheckingAllFiles = bInCheckingAllFiles;
@@ -433,6 +439,11 @@ public:
 		return bUpdateModifiedState;
 	}
 
+	bool ShouldUpdateModifiedStateToLocalRevision() const
+	{
+		return bUpdateModifiedStateToLocalRevision;
+	}
+
 	bool ShouldCheckAllFiles() const
 	{
 		return bCheckingAllFiles;
@@ -471,6 +482,9 @@ protected:
 	/** Whether to update the modified state - expensive */
 	bool bUpdateModifiedState;
 
+	/** Whether to update the modified state against local revision - only used when bUpdateModifiedState is true */
+	bool bUpdateModifiedStateToLocalRevision;
+
 	/** Hint that we are intending on checking all files in the project - some providers can optimize for this */
 	bool bCheckingAllFiles;
 
diff --git a/Engine/Source/Developer/UncontrolledChangelists/Private/UncontrolledChangelistState.cpp b/Engine/Source/Developer/UncontrolledChangelists/Private/UncontrolledChangelistState.cpp
--- a/Engine/Source/Developer/UncontrolledChangelists/Private/UncontrolledChangelistState.cpp
+++ b/Engine/Source/Developer/UncontrolledChangelists/Private/UncontrolledChangelistState.cpp
@@ -140,6 +140,7 @@ bool FUncontrolledChangelistState::AddFiles(const TArray<FString>& InFilenames,
 	{
 		auto UpdateStatusOperation = ISourceControlOperation::Create<FUpdateStatus>();
 		UpdateStatusOperation->SetUpdateModifiedState(true);
+		UpdateStatusOperation->SetUpdateModifiedStateToLocalRevision(true);
 		UpdateStatusOperation->SetQuiet(true);
 		UpdateStatusOperation->SetForceUpdate(true);
 		SourceControlProvider.Execute(UpdateStatusOperation, InFilenames);
