Fixed CIS for not finding the correct hash function of uintptr_t in SourceControlFileStatusMonitor.h

diff --git a/Engine/Source/Developer/SourceControl/Public/SourceControlFileStatusMonitor.h b/Engine/Source/Developer/SourceControl/Public/SourceControlFileStatusMonitor.h
--- a/Engine/Source/Developer/SourceControl/Public/SourceControlFileStatusMonitor.h
+++ b/Engine/Source/Developer/SourceControl/Public/SourceControlFileStatusMonitor.h
@@ -116,6 +116,15 @@ public:
 	void SetSuspendMonitoringPolicy(TFunction<bool()> IsMonitoringSuspended) { SuspendMonitoringPolicy = MoveTemp(IsMonitoringSuspended); }
 
 private:
+	// Some compilers were ambigous when picking a Hash function for uintprt_t, provide a custom one.
+	struct FUintptrHash : public TDefaultMapKeyFuncs<uintptr_t, FOnSourceControlFileStatus, false>
+	{
+		static uint32 GetKeyHash(const uintptr_t& Key)
+		{
+			return ::GetTypeHash(static_cast<uint64>(Key));
+		}
+	};
+
 	struct FSourceControlFileStatus
 	{
 		FSourceControlFileStatus(uintptr_t OriginalOwnerId, FOnSourceControlFileStatus InFileStatusUpdateDelegate)
@@ -125,7 +134,7 @@ private:
 
 		TSharedPtr<ISourceControlState> FileState;
 		double LastStatusCheckTimestampSecs = 0.0;
-		TMap<uintptr_t, FOnSourceControlFileStatus> OwnerDelegateMap;
+		TMap<uintptr_t, FOnSourceControlFileStatus, FDefaultSetAllocator, FUintptrHash> OwnerDelegateMap;
 	};
 
 private:
