- Remove UUnrealEdEngine::OnSourceControlFilesDeleted delegate (Unloading/Reloading of packages should be handled by USourceControlHelpers::ApplyOperationAndReloadPackages)
- Change ApplyOperationAndReloadPackages so that it calls DeleteObjectsUnchecked instead of UnloadPackage (which is what OnSourceControlFilesDeleted did)

#rnx

diff --git a/Engine/Source/Developer/SourceControl/Private/SourceControlHelpers.cpp b/Engine/Source/Developer/SourceControl/Private/SourceControlHelpers.cpp
--- a/Engine/Source/Developer/SourceControl/Private/SourceControlHelpers.cpp
+++ b/Engine/Source/Developer/SourceControl/Private/SourceControlHelpers.cpp
@@ -21,6 +21,7 @@
 #if WITH_EDITOR
 #include "Editor.h"
 #include "PackageTools.h"
+#include "ObjectTools.h"
 #endif
 
 #define LOCTEXT_NAMESPACE "SourceControlHelpers"
@@ -907,7 +908,7 @@ bool USourceControlHelpers::RevertFile(const FString& InFile, bool bSilent)
 bool USourceControlHelpers::ApplyOperationAndReloadPackages(const TArray<FString>& InFilenames, const TFunctionRef<bool(const TArray<FString>&)>& InOperation)
 {
 	ISourceControlProvider& SourceControlProvider = ISourceControlModule::Get().GetProvider();
-	TArray<TWeakObjectPtr<UPackage>> LoadedPackages;
+	TArray<UPackage*> LoadedPackages;
 	TArray<FString> PackageNames;
 	TArray<FString> PackageFilenames;
 	TArray<FString> FilteredActorPackages;
@@ -920,11 +921,11 @@ bool USourceControlHelpers::ApplyOperationAndReloadPackages(const TArray<FString
 		
 		if (FPackageName::TryConvertFilenameToLongPackageName(Filename, Result))
 		{
-			PackageNames.Emplace(MoveTemp(Result));
+			PackageNames.Add(MoveTemp(Result));
 		}
 		else
 		{
-			PackageNames.Emplace(Filename);
+			PackageNames.Add(Filename);
 		}
 	}
 
@@ -958,7 +959,7 @@ bool USourceControlHelpers::ApplyOperationAndReloadPackages(const TArray<FString
 	}
 
 	// Prepare the packages to be reverted...
-	for (TWeakObjectPtr<UPackage>& Package : LoadedPackages)
+	for (UPackage* Package : LoadedPackages)
 	{
 		// Detach the linkers of any loaded packages so that SCC can overwrite the files...
 		if (!Package->IsFullyLoaded())
@@ -966,7 +967,7 @@ bool USourceControlHelpers::ApplyOperationAndReloadPackages(const TArray<FString
 			FlushAsyncLoading();
 			Package->FullyLoad();
 		}
-		ResetLoaders(Package.Get());
+		ResetLoaders(Package);
 	}
 
 	PackageFilenames = SourceControlHelpers::PackageFilenames(PackageNames);
@@ -974,32 +975,45 @@ bool USourceControlHelpers::ApplyOperationAndReloadPackages(const TArray<FString
 	// Apply Operation
 	bSuccess = InOperation(PackageFilenames);
 
-	// Operation may have deleted some packages, so we need to unload those rather than re-load them...
-	TArray<UPackage*> PackagesToUnload;
-	TArray<UPackage*> PackagesToReload;
-	for (TWeakObjectPtr<UPackage>& Package : LoadedPackages)
+	// Reverting may have deleted some packages, so we need to delete those and unload them rather than re-load them...
+	FAssetRegistryModule& AssetRegistryModule = FModuleManager::LoadModuleChecked<FAssetRegistryModule>(TEXT("AssetRegistry"));
+	TArray<UObject*> ObjectsToDelete;
+	LoadedPackages.RemoveAll([&](UPackage* InPackage) -> bool
 	{
-		if (!Package.IsValid()) // Package was already unloaded on delete?
-		{
-			continue;
-		}
-
-		const FString PackageExtension = Package->ContainsMap() ? FPackageName::GetMapPackageExtension() : FPackageName::GetAssetPackageExtension();
-		const FString PackageFilename = FPackageName::LongPackageNameToFilename(Package->GetName(), PackageExtension);
+		const FString PackageExtension = InPackage->ContainsMap() ? FPackageName::GetMapPackageExtension() : FPackageName::GetAssetPackageExtension();
+		const FString PackageFilename = FPackageName::LongPackageNameToFilename(InPackage->GetName(), PackageExtension);
 		if (!FPaths::FileExists(PackageFilename))
 		{
-			PackagesToUnload.Emplace(Package.Get());
-		}
+			TArray<FAssetData> Assets;
+			AssetRegistryModule.Get().GetAssetsByPackageName(*InPackage->GetName(), Assets);
 
-		PackagesToReload.Emplace(Package.Get());
-	};
+			for (const FAssetData& Asset : Assets)
+			{
+				if (UObject* ObjectToDelete = Asset.FastGetAsset())
+				{
+					ObjectsToDelete.Add(ObjectToDelete);
+				}
+			}			
+			return true; // remove package
+		}
+		return false; // keep package
+	});
 
 	// Hot-reload the new packages...
-	UPackageTools::ReloadPackages(PackagesToReload);
-
-	// Unload any deleted packages...
-	UPackageTools::UnloadPackages(PackagesToUnload);
+	FText OutReloadErrorMsg;
+	const bool bInteractive = true;
+	UPackageTools::ReloadPackages(LoadedPackages, OutReloadErrorMsg, EReloadPackagesInteractionMode::Interactive);
+	if (!OutReloadErrorMsg.IsEmpty())
+	{
+		UE_LOG(LogSourceControl, Warning, TEXT("%s"), *OutReloadErrorMsg.ToString());
+	}
 
+	// Delete and Unload assets...
+	if(ObjectTools::DeleteObjectsUnchecked(ObjectsToDelete) != ObjectsToDelete.Num())
+	{ 
+		UE_LOG(LogSourceControl, Warning, TEXT("Failed to unload some assets."));
+	}
+	
 	// Re-cache the SCC state...
 	SourceControlProvider.Execute(ISourceControlOperation::Create<FUpdateStatus>(), PackageFilenames, EConcurrency::Asynchronous);
 
diff --git a/Engine/Source/Editor/UnrealEd/Classes/Editor/UnrealEdEngine.h b/Engine/Source/Editor/UnrealEd/Classes/Editor/UnrealEdEngine.h
--- a/Engine/Source/Editor/UnrealEd/Classes/Editor/UnrealEdEngine.h
+++ b/Engine/Source/Editor/UnrealEd/Classes/Editor/UnrealEdEngine.h
@@ -866,9 +866,6 @@ public:
 	/** Gets the project default map templates without any runtime overrides */
 	const TArray<FTemplateMapInfo>& GetProjectDefaultMapTemplates() const;
 
-	/** Called after files are deleted to perform necessary cleanups. */
-	void OnSourceControlFilesDeleted(const TArray<FString>& InDeletedFiles);
-
 protected:
 
 	/** Called when global editor selection changes */
diff --git a/Engine/Source/Editor/UnrealEd/Private/ObjectTools.cpp b/Engine/Source/Editor/UnrealEd/Private/ObjectTools.cpp
--- a/Engine/Source/Editor/UnrealEd/Private/ObjectTools.cpp
+++ b/Engine/Source/Editor/UnrealEd/Private/ObjectTools.cpp
@@ -2982,7 +2982,12 @@ namespace ObjectTools
 		{
 			GEditor->GetSelectedObjects()->Deselect(ObjectToDelete);
 		}
-		
+
+		if (ObjectToDelete->IsAsset())
+		{
+			GEditor->GetEditorSubsystem<UAssetEditorSubsystem>()->CloseAllEditorsForAsset(ObjectToDelete);
+		}
+
 		{
 			// @todo Animation temporary HACK to allow deleting of UMorphTargets. This will be removed when UMorphTargets are subobjects of USkeleton.
 			// Get the base skeleton and unregister this morphtarget
diff --git a/Engine/Source/Editor/UnrealEd/Private/UnrealEdEngine.cpp b/Engine/Source/Editor/UnrealEd/Private/UnrealEdEngine.cpp
--- a/Engine/Source/Editor/UnrealEd/Private/UnrealEdEngine.cpp
+++ b/Engine/Source/Editor/UnrealEd/Private/UnrealEdEngine.cpp
@@ -282,12 +282,6 @@ void UUnrealEdEngine::Init(IEngineLoop* InEngineLoop)
 			}
 		});
 #endif
-
-	// Delay this until after the source control module has loaded
-	FDelayedAutoRegisterHelper(EDelayedRegisterRunPhase::EarliestPossiblePluginsLoaded, [this]()
-	{
-		SourceControlFilesDeletedHandle = ISourceControlModule::Get().RegisterFilesDeleted(FSourceControlFilesDeletedDelegate::FDelegate::CreateUObject(this, &UUnrealEdEngine::OnSourceControlFilesDeleted));
-	});
 }
 
 bool CanCookForPlatformInThisProcess( const FString& PlatformName )
@@ -971,43 +965,6 @@ const TArray<FTemplateMapInfo>& UUnrealEdEngine::GetProjectDefaultMapTemplates()
 	return TemplateMapInfoCache;
 }
 
-void UUnrealEdEngine::OnSourceControlFilesDeleted(const TArray<FString>& InDeletedFiles)
-{
-	TArray<UObject*> ObjectsToDelete;
-
-	if (InDeletedFiles.IsEmpty())
-	{
-		return;
-	}
-
-	FAssetRegistryModule& AssetRegistryModule = FModuleManager::LoadModuleChecked<FAssetRegistryModule>(TEXT("AssetRegistry"));
-
-	for (const FString& File : InDeletedFiles)
-	{
-		TArray<FAssetData> Assets;
-		FString PackageName;
-
-		if (!FPackageName::TryConvertFilenameToLongPackageName(File, PackageName))
-		{
-			continue;
-		}
-
-		AssetRegistryModule.Get().GetAssetsByPackageName(*PackageName, Assets);
-
-		for (const FAssetData& Asset : Assets)
-		{
-			UObject* ObjectToDelete = Asset.GetAsset();
-
-			if (ObjectToDelete != nullptr)
-			{
-				ObjectsToDelete.Add(ObjectToDelete);
-			}
-		}
-	}
-
-	ObjectTools::DeleteObjectsUnchecked(ObjectsToDelete);
-}
-
 void UUnrealEdEngine::OnHISMTreeBuilt(UHierarchicalInstancedStaticMeshComponent* Component, bool bWasAsyncBuild)
 {
 	if (bWasAsyncBuild)
