Fixed crash in SourceControlFileStatusMonitor not checking of the Map iterator was valid before using it.

diff --git a/Engine/Source/Developer/SourceControl/Private/SourceControlFileStatusMonitor.cpp b/Engine/Source/Developer/SourceControl/Private/SourceControlFileStatusMonitor.cpp
--- a/Engine/Source/Developer/SourceControl/Private/SourceControlFileStatusMonitor.cpp
+++ b/Engine/Source/Developer/SourceControl/Private/SourceControlFileStatusMonitor.cpp
@@ -132,7 +132,7 @@ void FSourceControlFileStatusMonitor::StopMonitoringFiles(uintptr_t OwnerId, con
 
 void FSourceControlFileStatusMonitor::StopMonitoringFiles(uintptr_t OwnerId)
 {
-	for (auto It = MonitoredFiles.CreateIterator(); ;++It)
+	for (auto It = MonitoredFiles.CreateIterator(); It; ++It)
 	{
 		if (It->Value->OwnerDelegateMap.Remove(OwnerId) > 0 && It->Value->OwnerDelegateMap.IsEmpty())
 		{
