Fixed possible out-of-memory issue with the FSourceControlFileStatusMonitor reserving -1 items in an array.
  - Fixed an issue that didn't correctly account the number a pending files to check when source control provider changed.
  - Added check to be sure the monitor is called on the main thread.

diff --git a/Engine/Source/Developer/SourceControl/Private/SourceControlFileStatusMonitor.cpp b/Engine/Source/Developer/SourceControl/Private/SourceControlFileStatusMonitor.cpp
--- a/Engine/Source/Developer/SourceControl/Private/SourceControlFileStatusMonitor.cpp
+++ b/Engine/Source/Developer/SourceControl/Private/SourceControlFileStatusMonitor.cpp
@@ -38,6 +38,8 @@ TSharedPtr<FSourceControlFileStatusMonitor::FSourceControlFileStatus> FSourceCon
 
 void FSourceControlFileStatusMonitor::OnSourceControlProviderChanged(ISourceControlProvider& OldProvider, ISourceControlProvider& NewProvider)
 {
+	check(IsInGameThread()); // Concurrency issues if invoked from a background thread.
+		
 	// Start new with the new provider.
 	RequestedStatusFiles.Reset();
 	for (const TPair<FString, TSharedPtr<FSourceControlFileStatus>>& Pair : MonitoredFiles)
@@ -59,6 +61,8 @@ void FSourceControlFileStatusMonitor::OnSourceControlProviderChanged(ISourceCont
 
 void FSourceControlFileStatusMonitor::StartMonitoringFile(uintptr_t OwnerId, const FString& Pathname, FOnSourceControlFileStatus OnSourceControlFileStatus)
 {
+	check(IsInGameThread()); // Concurrency issues if invoked from a background thread.
+
 	// If the file is already monitored.
 	if (TSharedPtr<FSourceControlFileStatus> FileStatus = FindFileStatus(Pathname))
 	{
@@ -105,6 +109,8 @@ void FSourceControlFileStatusMonitor::StartMonitoringFiles(uintptr_t OwnerId, co
 
 void FSourceControlFileStatusMonitor::StopMonitoringFile(uintptr_t OwnerId, const FString& Pathname)
 {
+	check(IsInGameThread()); // Concurrency issues if the callback in invoked from a background thread.
+
 	if (TSharedPtr<FSourceControlFileStatus> FileStatus = FindFileStatus(Pathname))
 	{
 		if (FileStatus->OwnerDelegateMap.Remove(OwnerId) > 0 && FileStatus->OwnerDelegateMap.IsEmpty())
@@ -132,6 +138,8 @@ void FSourceControlFileStatusMonitor::StopMonitoringFiles(uintptr_t OwnerId, con
 
 void FSourceControlFileStatusMonitor::StopMonitoringFiles(uintptr_t OwnerId)
 {
+	check(IsInGameThread()); // Concurrency issues if the callback in invoked from a background thread.
+
 	for (auto It = MonitoredFiles.CreateIterator(); It; ++It)
 	{
 		if (It->Value->OwnerDelegateMap.Remove(OwnerId) > 0 && It->Value->OwnerDelegateMap.IsEmpty())
@@ -166,6 +174,8 @@ TOptional<FTimespan> FSourceControlFileStatusMonitor::GetStatusAge(const FString
 
 bool FSourceControlFileStatusMonitor::Tick(float DeltaTime)
 {
+	check(IsInGameThread()); // Concurrency issues if the callback in invoked from a background thread.
+
 	// Nothing to check or a request is already in-flight.
 	if (!ISourceControlModule::Get().IsEnabled() || MonitoredFiles.IsEmpty() || HasOngoingRequest())
 	{
@@ -225,11 +235,10 @@ bool FSourceControlFileStatusMonitor::Tick(float DeltaTime)
 	RequestedStatusFiles.Reserve(MaxFileNumPerRequestPolicy);
 	while (RequestedStatusFiles.Num() < MaxFileNumPerRequestPolicy)
 	{
-		if (NewFiles.Num())
+		if (NewFiles.Num() > 0)
 		{
 			RequestedStatusFiles.Emplace(NewFiles.Last()->Key);
 			NewFiles.Pop(/*bAllowShrinking*/false);
-			NewAddedFileCount--;
 		}
 		else if (RefreshedFiles.Num())
 		{
@@ -242,6 +251,9 @@ bool FSourceControlFileStatusMonitor::Tick(float DeltaTime)
 		}
 	}
 
+	// The remaining number of new files.
+	NewAddedFileCount = NewFiles.Num();
+
 	if (RequestedStatusFiles.Num())
 	{
 		LastSourceControlCheckSecs = NowSecs;
