Reduce a check to an ensure, as we are aware of some cases files are dirtied off the main thread.

This is a temporary stop gap until a better fix can be put in for the offending assets in a future release, and it is not intended for assets to dirty off the main thread.

#rnx

diff --git a/Engine/Source/Developer/SourceControl/Private/SourceControlFileStatusMonitor.cpp b/Engine/Source/Developer/SourceControl/Private/SourceControlFileStatusMonitor.cpp
--- a/Engine/Source/Developer/SourceControl/Private/SourceControlFileStatusMonitor.cpp
+++ b/Engine/Source/Developer/SourceControl/Private/SourceControlFileStatusMonitor.cpp
@@ -38,7 +38,7 @@ TSharedPtr<FSourceControlFileStatusMonitor::FSourceControlFileStatus> FSourceCon
 
 void FSourceControlFileStatusMonitor::OnSourceControlProviderChanged(ISourceControlProvider& OldProvider, ISourceControlProvider& NewProvider)
 {
-	check(IsInGameThread()); // Concurrency issues if invoked from a background thread.
+	ensure(IsInGameThread()); // Concurrency issues if invoked from a background thread.
 		
 	// Start new with the new provider.
 	RequestedStatusFiles.Reset();
@@ -61,7 +61,7 @@ void FSourceControlFileStatusMonitor::OnSourceControlProviderChanged(ISourceCont
 
 void FSourceControlFileStatusMonitor::StartMonitoringFile(uintptr_t OwnerId, const FString& Pathname, FOnSourceControlFileStatus OnSourceControlFileStatus)
 {
-	check(IsInGameThread()); // Concurrency issues if invoked from a background thread.
+	ensure(IsInGameThread()); // Concurrency issues if invoked from a background thread.
 
 	// If the file is already monitored.
 	if (TSharedPtr<FSourceControlFileStatus> FileStatus = FindFileStatus(Pathname))
@@ -109,7 +109,7 @@ void FSourceControlFileStatusMonitor::StartMonitoringFiles(uintptr_t OwnerId, co
 
 void FSourceControlFileStatusMonitor::StopMonitoringFile(uintptr_t OwnerId, const FString& Pathname)
 {
-	check(IsInGameThread()); // Concurrency issues if the callback in invoked from a background thread.
+	ensure(IsInGameThread()); // Concurrency issues if the callback in invoked from a background thread.
 
 	if (TSharedPtr<FSourceControlFileStatus> FileStatus = FindFileStatus(Pathname))
 	{
@@ -138,7 +138,7 @@ void FSourceControlFileStatusMonitor::StopMonitoringFiles(uintptr_t OwnerId, con
 
 void FSourceControlFileStatusMonitor::StopMonitoringFiles(uintptr_t OwnerId)
 {
-	check(IsInGameThread()); // Concurrency issues if the callback in invoked from a background thread.
+	ensure(IsInGameThread()); // Concurrency issues if the callback in invoked from a background thread.
 
 	for (auto It = MonitoredFiles.CreateIterator(); It; ++It)
 	{
@@ -174,7 +174,7 @@ TOptional<FTimespan> FSourceControlFileStatusMonitor::GetStatusAge(const FString
 
 bool FSourceControlFileStatusMonitor::Tick(float DeltaTime)
 {
-	check(IsInGameThread()); // Concurrency issues if the callback in invoked from a background thread.
+	ensure(IsInGameThread()); // Concurrency issues if the callback in invoked from a background thread.
 
 	// Nothing to check or a request is already in-flight.
 	if (!ISourceControlModule::Get().IsEnabled() || MonitoredFiles.IsEmpty() || HasOngoingRequest())
@@ -272,7 +272,7 @@ void FSourceControlFileStatusMonitor::OnSourceControlStatusUpdate(const TSharedR
 {
 	TRACE_CPUPROFILER_EVENT_SCOPE(FSourceControlFileStatusMonitor::OnSourceControlStatusUpdate);
 
-	check(IsInGameThread()); // Concurrency issues if the callback in invoked from a background thread.
+	ensure(IsInGameThread()); // Concurrency issues if the callback in invoked from a background thread.
 
 	double NowSecs = FPlatformTime::Seconds();
 
