Optimized the source control changelist to more efficienty update the list of source controlled files to monitor.
  - Added function to the source control file status monitor to easily set the exact list of file to monitor efficiently
Fixed the source control changelist not monitoring all the files in view.

diff --git a/Engine/Source/Developer/SourceControl/Private/SourceControlFileStatusMonitor.cpp b/Engine/Source/Developer/SourceControl/Private/SourceControlFileStatusMonitor.cpp
--- a/Engine/Source/Developer/SourceControl/Private/SourceControlFileStatusMonitor.cpp
+++ b/Engine/Source/Developer/SourceControl/Private/SourceControlFileStatusMonitor.cpp
@@ -27,9 +27,9 @@ FSourceControlFileStatusMonitor::~FSourceControlFileStatusMonitor()
 	FTSTicker::GetCoreTicker().RemoveTicker(TickerHandle);
 }
 
-TSharedPtr<FSourceControlFileStatusMonitor::FSourceControlFileStatus> FSourceControlFileStatusMonitor::FindFileStatus(const FString& Pathname) const
+TSharedPtr<FSourceControlFileStatusMonitor::FSourceControlFileStatus> FSourceControlFileStatusMonitor::FindFileStatus(const FString& AbsPathname) const
 {
-	if (const TSharedPtr<FSourceControlFileStatus>* FileStatus = MonitoredFiles.Find(Pathname))
+	if (const TSharedPtr<FSourceControlFileStatus>* FileStatus = MonitoredFiles.Find(AbsPathname))
 	{
 		return *FileStatus;
 	}
@@ -59,12 +59,12 @@ void FSourceControlFileStatusMonitor::OnSourceControlProviderChanged(ISourceCont
 	OldestFileStatusTimeSecs = 0.0;
 }
 
-void FSourceControlFileStatusMonitor::StartMonitoringFile(uintptr_t OwnerId, const FString& Pathname, FOnSourceControlFileStatus OnSourceControlFileStatus)
+void FSourceControlFileStatusMonitor::StartMonitoringFile(uintptr_t OwnerId, const FString& AbsPathname, FOnSourceControlFileStatus OnSourceControlFileStatus)
 {
 	ensure(IsInGameThread()); // Concurrency issues if invoked from a background thread.
 
 	// If the file is already monitored.
-	if (TSharedPtr<FSourceControlFileStatus> FileStatus = FindFileStatus(Pathname))
+	if (TSharedPtr<FSourceControlFileStatus> FileStatus = FindFileStatus(AbsPathname))
 	{
 		// Another 'client' is looking at this file (if the client already exit, override the callback with the new one).
 		FOnSourceControlFileStatus& OnSourceControlFileStatusDelegate = FileStatus->OwnerDelegateMap.FindOrAdd(OwnerId);
@@ -73,12 +73,12 @@ void FSourceControlFileStatusMonitor::StartMonitoringFile(uintptr_t OwnerId, con
 		// The monitor already knows the status.
 		if (FileStatus->FileState)
 		{
-			OnSourceControlFileStatusDelegate.ExecuteIfBound(Pathname, FileStatus->FileState.Get());
+			OnSourceControlFileStatusDelegate.ExecuteIfBound(AbsPathname, FileStatus->FileState.Get());
 		}
 	}
 	else
 	{
-		MonitoredFiles.Emplace(Pathname, MakeShared<FSourceControlFileStatus>(OwnerId, MoveTemp(OnSourceControlFileStatus)));
+		MonitoredFiles.Emplace(AbsPathname, MakeShared<FSourceControlFileStatus>(OwnerId, MoveTemp(OnSourceControlFileStatus)));
 		LastAddedFileTimeSecs = FPlatformTime::Seconds();
 		++NewAddedFileCount;
 	}
@@ -90,49 +90,48 @@ void FSourceControlFileStatusMonitor::StartMonitoringFile(uintptr_t OwnerId, con
 	}
 }
 
-void FSourceControlFileStatusMonitor::StartMonitoringFiles(uintptr_t OwnerId, const TArray<FString>& Pathnames, FOnSourceControlFileStatus OnSourceControlledFileStatus)
+void FSourceControlFileStatusMonitor::StartMonitoringFiles(uintptr_t OwnerId, const TArray<FString>& AbsPathnames, FOnSourceControlFileStatus OnSourceControlledFileStatus)
 {
-	for (const FString& Pathname : Pathnames)
+	for (const FString& AbsPathname : AbsPathnames)
 	{
-		StartMonitoringFile(OwnerId, Pathname, OnSourceControlledFileStatus);
+		StartMonitoringFile(OwnerId, AbsPathname, OnSourceControlledFileStatus);
 	}
 }
 
-void FSourceControlFileStatusMonitor::StartMonitoringFiles(uintptr_t OwnerId, const TSet<FString>& Pathnames, FOnSourceControlFileStatus OnSourceControlledFileStatus)
+void FSourceControlFileStatusMonitor::StartMonitoringFiles(uintptr_t OwnerId, const TSet<FString>& AbsPathnames, FOnSourceControlFileStatus OnSourceControlledFileStatus)
 {
-	for (const FString& Pathname : Pathnames)
+	for (const FString& AbsPathname : AbsPathnames)
 	{
-		StartMonitoringFile(OwnerId, Pathname, OnSourceControlledFileStatus);
+		StartMonitoringFile(OwnerId, AbsPathname, OnSourceControlledFileStatus);
 	}
 }
 
-
-void FSourceControlFileStatusMonitor::StopMonitoringFile(uintptr_t OwnerId, const FString& Pathname)
+void FSourceControlFileStatusMonitor::StopMonitoringFile(uintptr_t OwnerId, const FString& AbsPathname)
 {
 	ensure(IsInGameThread()); // Concurrency issues if the callback in invoked from a background thread.
 
-	if (TSharedPtr<FSourceControlFileStatus> FileStatus = FindFileStatus(Pathname))
+	if (TSharedPtr<FSourceControlFileStatus> FileStatus = FindFileStatus(AbsPathname))
 	{
 		if (FileStatus->OwnerDelegateMap.Remove(OwnerId) > 0 && FileStatus->OwnerDelegateMap.IsEmpty())
 		{
-			MonitoredFiles.Remove(Pathname);
+			MonitoredFiles.Remove(AbsPathname);
 		}
 	}
 }
 
-void FSourceControlFileStatusMonitor::StopMonitoringFiles(uintptr_t OwnerId, const TArray<FString>& Pathnames)
+void FSourceControlFileStatusMonitor::StopMonitoringFiles(uintptr_t OwnerId, const TArray<FString>& AbsPathnames)
 {
-	for (const FString& Pathname : Pathnames)
+	for (const FString& AbsPathname : AbsPathnames)
 	{
-		StopMonitoringFile(OwnerId, Pathname);
+		StopMonitoringFile(OwnerId, AbsPathname);
 	}
 }
 
-void FSourceControlFileStatusMonitor::StopMonitoringFiles(uintptr_t OwnerId, const TSet<FString>& Pathnames)
+void FSourceControlFileStatusMonitor::StopMonitoringFiles(uintptr_t OwnerId, const TSet<FString>& AbsPathnames)
 {
-	for (const FString& Pathname : Pathnames)
+	for (const FString& AbsPathname : AbsPathnames)
 	{
-		StopMonitoringFile(OwnerId, Pathname);
+		StopMonitoringFile(OwnerId, AbsPathname);
 	}
 }
 
@@ -149,6 +148,59 @@ void FSourceControlFileStatusMonitor::StopMonitoringFiles(uintptr_t OwnerId)
 	}
 }
 
+void FSourceControlFileStatusMonitor::SetMonitoringFiles(uintptr_t OwnerId, TSet<FString>&& AbsPathnames, FOnSourceControlFileStatus OnSourceControlledFileStatus)
+{
+	for (auto It = MonitoredFiles.CreateIterator(); It; ++It)
+	{
+		const FString& MonitoredAbsPathname = It->Key;
+		TSharedPtr<FSourceControlFileStatus>& MonitoredFileInfo = It->Value;
+		FOnSourceControlFileStatus* OwnerCallbackDelegate = MonitoredFileInfo->OwnerDelegateMap.Find(OwnerId);
+
+		// If the caller wants to monitor that file.
+		if (AbsPathnames.Contains(MonitoredAbsPathname))
+		{
+			// If the caller was already monitoring the file.
+			if (OwnerCallbackDelegate)
+			{
+				// Replace the delegate with the new one.
+				*OwnerCallbackDelegate = OnSourceControlledFileStatus;
+			}
+			else // The file was monitored, but not by this caller. Add this caller.
+			{
+				MonitoredFileInfo->OwnerDelegateMap.Add(OwnerId, OnSourceControlledFileStatus);
+
+				// If the status is already known.
+				if (MonitoredFileInfo->FileState)
+				{
+					OnSourceControlledFileStatus.ExecuteIfBound(MonitoredAbsPathname, MonitoredFileInfo->FileState.Get());
+				}
+			}
+
+			// That file was resolved, remove it from the set.
+			AbsPathnames.Remove(MonitoredAbsPathname);
+		}
+		else if (OwnerCallbackDelegate) // The caller used to monitor this file but doesn't want to monitor it anymore
+		{
+			if (MonitoredFileInfo->OwnerDelegateMap.Num() > 1)
+			{
+				// Just remove the caller, leaving the others.
+				MonitoredFileInfo->OwnerDelegateMap.Remove(OwnerId);
+			}
+			else
+			{
+				// Remove the caller and the files since nobody else monitor the file.
+				It.RemoveCurrent();
+			}
+		}
+	}
+
+	// What remains in the set are the files that weren't monitored yet, start monitoring them.
+	for (const FString& AbsPathname : AbsPathnames)
+	{
+		StartMonitoringFile(OwnerId, AbsPathname, OnSourceControlledFileStatus);
+	}
+}
+
 TSet<FString> FSourceControlFileStatusMonitor::GetMonitoredFiles(uintptr_t OwnerId)
 {
 	TSet<FString> Pathnames;
@@ -162,10 +214,10 @@ TSet<FString> FSourceControlFileStatusMonitor::GetMonitoredFiles(uintptr_t Owner
 	return Pathnames;
 }
 
-TOptional<FTimespan> FSourceControlFileStatusMonitor::GetStatusAge(const FString& Pathname) const
+TOptional<FTimespan> FSourceControlFileStatusMonitor::GetStatusAge(const FString& AbsPathname) const
 {
 	TOptional<FTimespan> Age;
-	if (TSharedPtr<FSourceControlFileStatus> FileStatus = FindFileStatus(Pathname))
+	if (TSharedPtr<FSourceControlFileStatus> FileStatus = FindFileStatus(AbsPathname))
 	{
 		Age.Emplace(FTimespan::FromSeconds(FileStatus->FileState ? FPlatformTime::Seconds() - FileStatus->LastStatusCheckTimestampSecs : 0.0));
 	}
@@ -292,19 +344,19 @@ void FSourceControlFileStatusMonitor::OnSourceControlStatusUpdate(const TSharedR
 		return;
 	}
 
-	for (const FString& Pathname : RequestedStatusFiles)
+	for (const FString& AbsPathname : RequestedStatusFiles)
 	{
 		// NOTE: The file and its status can be removed while exeucting OnFileStatusUpdateDelegate. Keep the shared pointer to avoid early destruction.
-		if (TSharedPtr<FSourceControlFileStatus> FileStatus = FindFileStatus(Pathname))
+		if (TSharedPtr<FSourceControlFileStatus> FileStatus = FindFileStatus(AbsPathname))
 		{
-			if (TSharedPtr<ISourceControlState> FileState = ISourceControlModule::Get().GetProvider().GetState(Pathname, EStateCacheUsage::Use))
+			if (TSharedPtr<ISourceControlState> FileState = ISourceControlModule::Get().GetProvider().GetState(AbsPathname, EStateCacheUsage::Use))
 			{
 				FileStatus->FileState = MoveTemp(FileState);
 				FileStatus->LastStatusCheckTimestampSecs = NowSecs;
 
 				for (TPair<uintptr_t, FOnSourceControlFileStatus>& OwnerDelegatePair : FileStatus->OwnerDelegateMap)
 				{
-					OwnerDelegatePair.Value.ExecuteIfBound(Pathname, FileStatus->FileState.Get());
+					OwnerDelegatePair.Value.ExecuteIfBound(AbsPathname, FileStatus->FileState.Get());
 				}
 			}
 		}
diff --git a/Engine/Source/Developer/SourceControl/Public/SourceControlFileStatusMonitor.h b/Engine/Source/Developer/SourceControl/Public/SourceControlFileStatusMonitor.h
--- a/Engine/Source/Developer/SourceControl/Public/SourceControlFileStatusMonitor.h
+++ b/Engine/Source/Developer/SourceControl/Public/SourceControlFileStatusMonitor.h
@@ -23,7 +23,7 @@ class SOURCECONTROL_API FSourceControlFileStatusMonitor : public TSharedFromThis
 {
 public:
 	/** Invoked when the status of a source controlled files is updated. The state can be null when the source control provider is changed. */
-	DECLARE_DELEGATE_TwoParams(FOnSourceControlFileStatus, const FString& /*Pathname*/, const ISourceControlState* /*State*/);
+	DECLARE_DELEGATE_TwoParams(FOnSourceControlFileStatus, const FString& /*AbsPathname*/, const ISourceControlState* /*State*/);
 
 public:
 	FSourceControlFileStatusMonitor();
@@ -39,7 +39,7 @@ public:
 	 * within the callback OnSourceControlledFileStatus().
 	 *
 	 * @param Owner The unique Id of the caller, typically the caller memory address.
-	 * @param Pathname The absolute path and filname of the file to monitor.
+	 * @param AbsPathname The absolute path and filname of the file to monitor.
 	 * @param OnSourceControlledFileStatus Delegate invoked whe the status of the file is updated.
 	 *
 	 * @see SetNewRequestProbationPeriod()
@@ -47,28 +47,38 @@ public:
 	 * @see SetSuspendMonitoringPolicy()
 	 * @see GetStatusAge()
 	 */
-	void StartMonitoringFile(uintptr_t OwnerId, const FString& Pathname, FOnSourceControlFileStatus OnSourceControlledFileStatus);
-	void StartMonitoringFiles(uintptr_t OwnerId, const TArray<FString>& Pathnames, FOnSourceControlFileStatus OnSourceControlledFileStatus);
-	void StartMonitoringFiles(uintptr_t OwnerId, const TSet<FString>& Pathnames, FOnSourceControlFileStatus OnSourceControlledFileStatus);
+	void StartMonitoringFile(uintptr_t OwnerId, const FString& AbsPathname, FOnSourceControlFileStatus OnSourceControlledFileStatus);
+	void StartMonitoringFiles(uintptr_t OwnerId, const TArray<FString>& AbsPathnames, FOnSourceControlFileStatus OnSourceControlledFileStatus);
+	void StartMonitoringFiles(uintptr_t OwnerId, const TSet<FString>& AbsPathnames, FOnSourceControlFileStatus OnSourceControlledFileStatus);
 
 	/**
 	 * Stops monitoring the source control status of the specified file. If the specified file is not monitored, the function
 	 * returns successfully.
 	 *
 	 * @param OwnerId The unique Id of the caller, typically the caller memory address.
-	 * @param Pathname The absolute path and filname of the file to stop monitoring.
+	 * @param AbsPathname The absolute path and filname of the file to stop monitoring.
 	 *
 	 * @note This can be called fron the FOnSourceControlFileStatus callback passed to StartMonitoringFile().
 	 */
-	void StopMonitoringFile(uintptr_t OwnerId, const FString& Pathname);
-	void StopMonitoringFiles(uintptr_t OwnerId, const TArray<FString>& Pathnames);
-	void StopMonitoringFiles(uintptr_t OwnerId, const TSet<FString>& Pathnames);
+	void StopMonitoringFile(uintptr_t OwnerId, const FString& AbsPathname);
+	void StopMonitoringFiles(uintptr_t OwnerId, const TArray<FString>& AbsPathnames);
+	void StopMonitoringFiles(uintptr_t OwnerId, const TSet<FString>& AbsPathnames);
 
 	/**
 	 * Stops monitoring all the files that were registered by the specified owner Id.
 	 */
 	void StopMonitoringFiles(uintptr_t OwnerId);
 
+	/**
+	 * Starts monitoring files that weren't monitored yet by the specified owner, stops monitoring those that were monitored by the owner but are not in the
+	 * updated list and keep monitoring files that were monitored before and still in the updated list.
+	 *
+	 * @param OwnerId The unique Id of the caller, typically the caller memory address.
+	 * @param AbsPathnames The list of absolute file pathnames that must be monitored. (The set is modified during the operation)
+	 * @param OnSourceControlledFileStatus Delegate invoked whe the status of the file is updated.
+	 */
+	void SetMonitoringFiles(uintptr_t OwnerId, TSet<FString>&& AbsPathnames, FOnSourceControlFileStatus OnSourceControlledFileStatus);
+
 	/**
 	 * Returns the set of files being monitored by the specifed owner.
 	 */
@@ -79,7 +89,7 @@ public:
 	 * monitored. If the initial status request hasn't been received yet, returns
 	 * zero.
 	 */
-	TOptional<FTimespan> GetStatusAge(const FString& Pathname) const;
+	TOptional<FTimespan> GetStatusAge(const FString& AbsPathname) const;
 
 	/**
 	 * Set the times waited before a newly added file status request is sent to the source control provider. This is used batch several
@@ -154,7 +164,7 @@ private:
 	bool HasOngoingRequest() const { return RequestedStatusFiles.Num() > 0; }
 
 private:
-	/** The files currently monitored. */
+	/** The files currently monitored, mapping absolute file pathname to the monitored information. */
 	TMap<FString, TSharedPtr<FSourceControlFileStatus>> MonitoredFiles;
 
 	/** The list of files for which a request to the source control provider is currently in-flight/ongoing. */
diff --git a/Engine/Source/Editor/SourceControlWindows/Private/SSourceControlChangelists.cpp b/Engine/Source/Editor/SourceControlWindows/Private/SSourceControlChangelists.cpp
--- a/Engine/Source/Editor/SourceControlWindows/Private/SSourceControlChangelists.cpp
+++ b/Engine/Source/Editor/SourceControlWindows/Private/SSourceControlChangelists.cpp
@@ -658,12 +658,13 @@ void SSourceControlChangelistsWidget::Tick(const FGeometry& AllottedGeometry, co
 				// Minimize the list of file status monitored to only include those in view. The ones outside of the visible area don't need
 				// to be monitored unless the user sorts.
 				int32 ItemStart = static_cast<int32>(FileListView->GetScrollOffset()); // This give the offset in item count.
+				int32 ItemEnd = ItemStart + ItemCount;
 				TSet<FString> MonitoredFiles;
-				for (int32 Index = ItemStart; Index < ItemCount; ++Index)
+				for (int32 Index = ItemStart; Index < ItemEnd; ++Index)
 				{
 					MonitoredFiles.Emplace(StaticCastSharedPtr<IFileViewTreeItem>(FileListNodes[Index])->GetFullPathname());
 				}
-				RequestFileStatusRefresh(MonitoredFiles);
+				RequestFileStatusRefresh(MoveTemp(MonitoredFiles));
 				bUpdateMonitoredFileStatusList = false;
 			}
 		}
@@ -781,11 +782,11 @@ void SSourceControlChangelistsWidget::RequestFileStatusRefresh(const IChangelist
 	if (!Pathnames.IsEmpty())
 	{
 		// Fire an async task to get the latest files status from source control to get the 'Checked Out By' status.
-		RequestFileStatusRefresh(Pathnames);
+		RequestFileStatusRefresh(MoveTemp(Pathnames));
 	}
 }
 
-void SSourceControlChangelistsWidget::RequestFileStatusRefresh(const TSet<FString>& PathnamesToMonitor)
+void SSourceControlChangelistsWidget::RequestFileStatusRefresh(TSet<FString>&& PathnamesToMonitor)
 {
 	TRACE_CPUPROFILER_EVENT_SCOPE(SSourceControlChangelistsWidget::RequestFileStatusRefresh);
 
@@ -801,12 +802,7 @@ void SSourceControlChangelistsWidget::RequestFileStatusRefresh(const TSet<FStrin
 	if (!PathnamesToMonitor.IsEmpty())
 	{
 		// Request an update only for the files that changed.
-		TSet<FString> MonitoredPathnames = ISourceControlModule::Get().GetSourceControlFileStatusMonitor().GetMonitoredFiles(reinterpret_cast<uintptr_t>(this));
-		TSet<FString> StopMonitoringPathnames = MonitoredPathnames.Difference(PathnamesToMonitor);
-		TSet<FString> StartMonitoringPathnames = PathnamesToMonitor.Difference(MonitoredPathnames);
-
-		ISourceControlModule::Get().GetSourceControlFileStatusMonitor().StopMonitoringFiles(reinterpret_cast<uintptr_t>(this), StopMonitoringPathnames);
-		ISourceControlModule::Get().GetSourceControlFileStatusMonitor().StartMonitoringFiles(reinterpret_cast<uintptr_t>(this), StartMonitoringPathnames, FSourceControlFileStatusMonitor::FOnSourceControlFileStatus());
+		ISourceControlModule::Get().GetSourceControlFileStatusMonitor().SetMonitoringFiles(reinterpret_cast<uintptr_t>(this), MoveTemp(PathnamesToMonitor), FSourceControlFileStatusMonitor::FOnSourceControlFileStatus());
 		// NOTE: The updated status are expected to come from OnSourceControlStateChanged() that is invoked when the source control internal state changes.
 	}
 	else
diff --git a/Engine/Source/Editor/SourceControlWindows/Private/SSourceControlChangelists.h b/Engine/Source/Editor/SourceControlWindows/Private/SSourceControlChangelists.h
--- a/Engine/Source/Editor/SourceControlWindows/Private/SSourceControlChangelists.h
+++ b/Engine/Source/Editor/SourceControlWindows/Private/SSourceControlChangelists.h
@@ -85,7 +85,7 @@ private:
 
 	void RequestChangelistsRefresh();
 	void RequestFileStatusRefresh(const IChangelistTreeItem& Changelist);
-	void RequestFileStatusRefresh(const TSet<FString>& Pathnames);
+	void RequestFileStatusRefresh(TSet<FString>&& Pathnames);
 	void OnRefreshUI(ERefreshFlags RefreshFlags);
 	void OnRefreshSourceControlWidgets(int64 CurreUpdateNum, const TFunction<void(TSharedPtr<IFileViewTreeItem>&)>& AddItemToFileView);
 	void OnRefreshUncontrolledChangelistWidgets(int64 CurreUpdateNum, const TFunction<void(TSharedPtr<IFileViewTreeItem>&)>& AddItemToFileView);
