Allow FGetChangelistDetails to accept the changelist number as part of the constructor, allowing it to be passed in as an arg for ISourceControlOperation::Create

#rnx

- Added a FStringView and FString&& versions of the constructor so we should be covered for all use cases.

diff --git a/Engine/Source/Developer/SourceControl/Public/SourceControlOperations.h b/Engine/Source/Developer/SourceControl/Public/SourceControlOperations.h
--- a/Engine/Source/Developer/SourceControl/Public/SourceControlOperations.h
+++ b/Engine/Source/Developer/SourceControl/Public/SourceControlOperations.h
@@ -605,6 +605,22 @@ public:
 class FGetChangelistDetails : public FSourceControlOperationBase
 {
 public:
+	FGetChangelistDetails() = default;
+
+	explicit FGetChangelistDetails(FString&& InChangelistNumber)
+		: ChangelistNumber(MoveTemp(InChangelistNumber))
+	{
+
+	}
+
+	explicit FGetChangelistDetails(FStringView InChangelistNumber)
+		: ChangelistNumber(InChangelistNumber)
+	{
+
+	}
+
+	virtual ~FGetChangelistDetails() = default;
+
 	// ISourceControlOperation interface
 	virtual FName GetName() const override
 	{
