Avoid getting spam warnings coming from world outliner usaved column trying to resolve actor memory package path when running PIE.

diff --git a/Engine/Source/Developer/SourceControl/Private/SourceControlHelpers.cpp b/Engine/Source/Developer/SourceControl/Private/SourceControlHelpers.cpp
--- a/Engine/Source/Developer/SourceControl/Private/SourceControlHelpers.cpp
+++ b/Engine/Source/Developer/SourceControl/Private/SourceControlHelpers.cpp
@@ -1465,7 +1465,16 @@ FString USourceControlHelpers::PackageFilename( const UPackage* InPackage )
 	FString Filename;
 	if(InPackage != nullptr)
 	{
-		Filename = FPaths::ConvertRelativePathToFull(PackageFilename_Internal(InPackage->GetName()));
+		// Prefer using package loaded path to resolve file name as it properly resolves memory packages
+		if (!InPackage->GetLoadedPath().IsEmpty())
+		{
+			const FString PackageExtension = InPackage->ContainsMap() ? FPackageName::GetMapPackageExtension() : FPackageName::GetAssetPackageExtension();
+			Filename = FPaths::ConvertRelativePathToFull(FPackageName::LongPackageNameToFilename(InPackage->GetLoadedPath().GetPackageName(), PackageExtension));
+		}
+		else
+		{
+			Filename = FPaths::ConvertRelativePathToFull(PackageFilename_Internal(InPackage->GetName()));
+		}
 	}
 	return Filename;
 }
