Added separate toast description for generic error on failed auto-checkout

diff --git a/Engine/Source/Developer/SourceControl/Public/ISourceControlOperation.h b/Engine/Source/Developer/SourceControl/Public/ISourceControlOperation.h
--- a/Engine/Source/Developer/SourceControl/Public/ISourceControlOperation.h
+++ b/Engine/Source/Developer/SourceControl/Public/ISourceControlOperation.h
@@ -13,6 +13,7 @@ struct FSourceControlResultInfo
 	{
 		InfoMessages.Append(InResultInfo.InfoMessages);
 		ErrorMessages.Append(InResultInfo.ErrorMessages);
+		Tags.Append(InResultInfo.Tags);
 	}
 
 	/** Info and/or warning message storage */
@@ -20,6 +21,9 @@ struct FSourceControlResultInfo
 
 	/** Potential error message storage */
 	TArray<FText> ErrorMessages;
+
+	/** Additional arbitrary information attached to the command */
+	TArray<FString> Tags;
 };
 
 
@@ -60,6 +64,12 @@ public:
 		// Implemented in subclasses
 	}
 
+	/** Add tag. */
+	virtual void AddTag(const FString& InTag)
+	{
+		// Implemented in subclasses
+	}
+
 	/**
 	 * Append any info or error messages that may have accumulated during the operation prior
 	 * to returning a result, ensuring to keep any already accumulated info.
