Implemented Asset Revert functionality from Scene Outliner.
Fixed 'Cannot revert umaps with Skein' by reloading all assets when attempting to revert a map
Implemented 'Revert All' button

diff --git a/Engine/Plugins/Editor/AssetManagerEditor/Source/AssetManagerEditor/Private/AssetSourceControlContextMenu.cpp b/Engine/Plugins/Editor/AssetManagerEditor/Source/AssetManagerEditor/Private/AssetSourceControlContextMenu.cpp
--- a/Engine/Plugins/Editor/AssetManagerEditor/Source/AssetManagerEditor/Private/AssetSourceControlContextMenu.cpp
+++ b/Engine/Plugins/Editor/AssetManagerEditor/Source/AssetManagerEditor/Private/AssetSourceControlContextMenu.cpp
@@ -777,7 +777,17 @@ void FAssetSourceControlContextMenuState::ExecuteSCCRevert() const
 {
 	TArray<FString> PackageNames;
 	GetSelectedPackageNames(PackageNames);
-	FSourceControlWindows::PromptForRevert(PackageNames);
+	bool bReloadAll = false;
+	for (FString PackageName : PackageNames)
+	{
+		UPackage* Package = FindPackage(nullptr, *PackageName);
+		if (Package->ContainsMap())
+		{
+			bReloadAll = true;
+			break;
+		}
+	}
+	FSourceControlWindows::PromptForRevert(PackageNames, bReloadAll);
 }
 
 void FAssetSourceControlContextMenuState::ExecuteSCCRevertWritable() const
diff --git a/Engine/Source/Developer/SourceControl/Private/SourceControlHelpers.cpp b/Engine/Source/Developer/SourceControl/Private/SourceControlHelpers.cpp
--- a/Engine/Source/Developer/SourceControl/Private/SourceControlHelpers.cpp
+++ b/Engine/Source/Developer/SourceControl/Private/SourceControlHelpers.cpp
@@ -138,7 +138,7 @@ FString ConvertFileToQualifiedPath(const FString& InFile, bool bSilent, bool bAl
 					SCFile /= FString();
 				}
 				else if (AssociatedExtension)
-			{
+				{
 					// Just use the requested extension
 					SCFile += AssociatedExtension;
 				}
@@ -147,10 +147,10 @@ FString ConvertFileToQualifiedPath(const FString& InFile, bool bSilent, bool bAl
 					// The package does not exist on disk, see if we can find it in memory and predict the file extension
 					UPackage* Package = FindPackage(nullptr, *SCFile);
 					SCFile += (Package && Package->ContainsMap() ? FPackageName::GetMapPackageExtension() : FPackageName::GetAssetPackageExtension());
-					}
 				}
-				else
-				{
+			}
+			else
+			{
 				bPackage = false;
 			}
 		}
@@ -907,42 +907,79 @@ bool USourceControlHelpers::RevertFile(const FString& InFile, bool bSilent)
 }
 
 #if WITH_EDITOR
-bool USourceControlHelpers::ApplyOperationAndReloadPackages(const TArray<FString>& InFilenames, const TFunctionRef<bool(const TArray<FString>&)>& InOperation)
+bool USourceControlHelpers::ApplyOperationAndReloadPackages(const TArray<FString>& InPackagesToApplyOperation, const TArray<FString>& InPackagesToReload, 
+	const TFunctionRef<bool(const TArray<FString>&)>& InOperation, bool bAllowReloadWorld)
 {
 	ISourceControlProvider& SourceControlProvider = ISourceControlModule::Get().GetProvider();
 	TArray<UPackage*> LoadedPackages;
-	TArray<FString> PackageNames;
-	TArray<FString> PackageFilenames;
+	TArray<FString> PackageNamesToApplyOperation;
+	TArray<FString> PackageFilenamesToApplyOperation;
+	TArray<FString> PackageNamesToReload;
 	TArray<FString> FilteredActorPackages;
 	bool bSuccess = false;
 
+	auto DetachLinker = [](UPackage* Package)
+	{
+		if (!Package->IsFullyLoaded())
+		{
+			FlushAsyncLoading();
+			Package->FullyLoad();
+		}
+		ResetLoaders(Package);
+	};
+
 	// Normalize packagenames and filenames
-	for (const FString& Filename : InFilenames)
+	for (const FString& Filename : InPackagesToApplyOperation)
 	{
 		FString Result;
 		
 		if (FPackageName::TryConvertFilenameToLongPackageName(Filename, Result))
 		{
-			PackageNames.Add(MoveTemp(Result));
+			PackageNamesToApplyOperation.Add(MoveTemp(Result));
+		}
+		else
+		{
+			PackageNamesToApplyOperation.Add(Filename);
+		}
+	}
+
+	for (const FString& Filename : InPackagesToReload)
+	{
+		FString Result;
+
+		if (FPackageName::TryConvertFilenameToLongPackageName(Filename, Result))
+		{
+			PackageNamesToReload.Add(MoveTemp(Result));
 		}
 		else
 		{
-			PackageNames.Add(Filename);
+			PackageNamesToReload.Add(Filename);
 		}
 	}
 
 	// Remove packages if they are loaded actors or world
-	PackageNames.RemoveAll([&FilteredActorPackages, &LoadedPackages](const FString& PackageName) -> bool
+	PackageNamesToReload.RemoveAll([&FilteredActorPackages, &LoadedPackages, bAllowReloadWorld, &DetachLinker](const FString& PackageName) -> bool
 	{
 		UPackage* Package = FindPackage(NULL, *PackageName);
 		
 		if (Package != nullptr)
 		{
-			if (UWorld::IsWorldOrExternalActorPackage(Package))
+			if (UWorld* World = UWorld::FindWorldInPackage(Package); World && !bAllowReloadWorld)
 			{
 				FilteredActorPackages.Emplace(PackageName);
 				return true; // remove the package
 			}
+			else if (AActor* Actor = AActor::FindActorInPackage(Package))
+			{
+				if (bAllowReloadWorld)
+				{
+					DetachLinker(Package);
+					return true;
+				}
+
+				FilteredActorPackages.Emplace(PackageName);
+				return true; // remove the package
+			}
 
 			LoadedPackages.Add(Package);
 		}
@@ -964,18 +1001,13 @@ bool USourceControlHelpers::ApplyOperationAndReloadPackages(const TArray<FString
 	for (UPackage* Package : LoadedPackages)
 	{
 		// Detach the linkers of any loaded packages so that SCC can overwrite the files...
-		if (!Package->IsFullyLoaded())
-		{
-			FlushAsyncLoading();
-			Package->FullyLoad();
-		}
-		ResetLoaders(Package);
+		DetachLinker(Package);
 	}
 
-	PackageFilenames = SourceControlHelpers::PackageFilenames(PackageNames);
+	PackageFilenamesToApplyOperation = SourceControlHelpers::PackageFilenames(PackageNamesToApplyOperation);
 
 	// Apply Operation
-	bSuccess = InOperation(PackageFilenames);
+	bSuccess = InOperation(PackageFilenamesToApplyOperation);
 
 	// Reverting may have deleted some packages, so we need to delete those and unload them rather than re-load them...
 	FAssetRegistryModule& AssetRegistryModule = FModuleManager::LoadModuleChecked<FAssetRegistryModule>(TEXT("AssetRegistry"));
@@ -1017,14 +1049,114 @@ bool USourceControlHelpers::ApplyOperationAndReloadPackages(const TArray<FString
 	}
 	
 	// Re-cache the SCC state...
-	SourceControlProvider.Execute(ISourceControlOperation::Create<FUpdateStatus>(), PackageFilenames, EConcurrency::Asynchronous);
+	if (bAllowReloadWorld)
+	{
+		auto UpdateCommand = ISourceControlOperation::Create<FUpdateStatus>();
+		UpdateCommand->SetCheckingAllFiles(true);
+		UpdateCommand->SetForceUpdate(true);
+		SourceControlProvider.Execute(UpdateCommand, EConcurrency::Asynchronous);
+	}
+	else
+	{
+		SourceControlProvider.Execute(ISourceControlOperation::Create<FUpdateStatus>(), PackageFilenamesToApplyOperation, EConcurrency::Asynchronous);
+	}
 
 	return bSuccess;
 }
 
-bool USourceControlHelpers::RevertAndReloadPackages(const TArray<FString>& InFilenames)
+TArray<FString> USourceControlHelpers::GetSourceControlLocations(const bool bContentOnly)
+{
+	TArray<FString> SourceControlLocations;
+
+	{
+		TArray<FString> RootPaths;
+		FPackageName::QueryRootContentPaths(RootPaths);
+		for (const FString& RootPath : RootPaths)
+		{
+			const FString RootPathOnDisk = FPackageName::LongPackageNameToFilename(RootPath);
+			SourceControlLocations.Add(FPaths::ConvertRelativePathToFull(RootPathOnDisk));
+		}
+	}
+
+	if (!bContentOnly)
+	{
+		SourceControlLocations.Add(FPaths::ConvertRelativePathToFull(FPaths::ProjectConfigDir()));
+		SourceControlLocations.Add(FPaths::ConvertRelativePathToFull(FPaths::GetProjectFilePath()));
+	}
+
+	return SourceControlLocations;
+}
+
+// TODO: Move outside of helpers
+bool USourceControlHelpers::ListAllPackages(TArray<FString>& OutPackageNames)
+{
+	// determine source control locations
+	TArray<FString> SourceControlLocations;
+
+	if (ISourceControlModule::Get().UsesCustomProjectDir())
+	{
+		SourceControlLocations.Add(ISourceControlModule::Get().GetSourceControlProjectDir());
+	}
+	else
+	{
+		GetSourceControlLocations(/*bContentOnly=*/true);
+	}
+
+	// find packages in those locations
+	for (const FString& SourceControlLocation : SourceControlLocations)
+	{
+		TArray<FString> PackageRelativePaths;
+		if (FPackageName::FindPackagesInDirectory(PackageRelativePaths, FPaths::ConvertRelativePathToFull(SourceControlLocation)))
+		{
+			OutPackageNames.Reserve(OutPackageNames.Num() + PackageRelativePaths.Num());
+			for (const FString& Path : PackageRelativePaths)
+			{
+				FString PackageName;
+				FString FailureReason;
+				if (FPackageName::TryConvertFilenameToLongPackageName(Path, PackageName, &FailureReason))
+				{
+					OutPackageNames.Add(PackageName);
+				}
+				else
+				{
+					FMessageLog("SourceControl").Error(FText::FromString(FailureReason));
+				}
+			}
+		}
+	}
+
+	return true;
+}
+
+// TODO: Move outside of helpers
+bool USourceControlHelpers::RevertAndReloadAllPackages(const TArray<FString>& InFilenames)
 {
-	auto RevertOperation = [](const TArray<FString>& InFilenames) -> bool
+	TArray<FString> PackagesToReload;
+	ListAllPackages(PackagesToReload);
+
+	if (InFilenames.Num() > 0)
+	{
+		RevertAndReloadPackages(InFilenames, PackagesToReload, /*bRevertAll=*/false, /*bAllowReloadWorld=*/true);
+		return true;
+	}
+
+	return false;
+}
+
+// TODO: Move outside of helpers
+bool USourceControlHelpers::RevertAllChangesAndReloadAllPackages()
+{
+	TArray<FString> PackagesToReload;
+	ListAllPackages(PackagesToReload);
+
+	RevertAndReloadPackages(TArray<FString>(), PackagesToReload, /*bRevertAll=*/true, /*bAllowReloadWorld=*/true);
+
+	return true;
+}
+
+bool USourceControlHelpers::RevertAndReloadPackages(const TArray<FString>& InPackagesToRevert, const TArray<FString>& InPackagesToReload, bool bRevertAll, bool bAllowReloadWorld)
+{
+	auto RevertOperation = [bRevertAll](const TArray<FString>& InPackagesToRevert) -> bool
 	{
 		ISourceControlProvider& SourceControlProvider = ISourceControlModule::Get().GetProvider();
 
@@ -1037,10 +1169,17 @@ bool USourceControlHelpers::RevertAndReloadPackages(const TArray<FString>& InFil
 			}
 		});
 
-		return SourceControlProvider.Execute(ISourceControlOperation::Create<FRevert>(), InFilenames, EConcurrency::Synchronous, OperationCompleteCallback) == ECommandResult::Succeeded;
+		auto RevertOperation = ISourceControlOperation::Create<FRevert>();
+
+		if (bRevertAll)
+		{
+			RevertOperation->SetRevertAll(true);
+		}
+
+		return SourceControlProvider.Execute(RevertOperation, InPackagesToRevert, EConcurrency::Synchronous, OperationCompleteCallback) == ECommandResult::Succeeded;
 	};
 
-	return ApplyOperationAndReloadPackages(InFilenames, RevertOperation);
+	return ApplyOperationAndReloadPackages(InPackagesToRevert, InPackagesToReload, RevertOperation, bAllowReloadWorld);
 }
 #endif //!WITH_EDITOR
 
diff --git a/Engine/Source/Developer/SourceControl/Public/SourceControlHelpers.h b/Engine/Source/Developer/SourceControl/Public/SourceControlHelpers.h
--- a/Engine/Source/Developer/SourceControl/Public/SourceControlHelpers.h
+++ b/Engine/Source/Developer/SourceControl/Public/SourceControlHelpers.h
@@ -351,19 +351,38 @@ public:
 
 	/**
 	 * Applies the provided function and reloads provided packages.
-	 * @param	InFilenames		The files/packages to apply the operation and reload
-	 * @param	InOperation		The function to apply
+	 * @param	InPackagesToApplyOperation			The files/packages to apply the operation
+	 * @param	InPackagesToReload					The files to reload
+	 * @param	InOperation							The function to apply
+	 * @param	bAllowReloadWorld					Allow reloading of the world
 	 * @return true if succeeded.
 	 */
-	static bool ApplyOperationAndReloadPackages(const TArray<FString>& InFilenames, const TFunctionRef<bool(const TArray<FString>&)>& InOperation);
+	static bool ApplyOperationAndReloadPackages(const TArray<FString>& InPackagesToApplyOperation, const TArray<FString>& InPackagesToReload, 
+		const TFunctionRef<bool(const TArray<FString>&)>& InOperation, bool bAllowReloadWorld = false);
 
 	/**
      * Reverts the provided files then reload packages.
-     * @param	InFilenames		The files/packages to revert and reload
+	 * @param	InPackagesToRevert					The packages to revert
+	 * @param	InPackagesToReload					The packages to reload
+	 * @param	bRevertAll							Whether to revert all files
+	 * @param	bAllowReloadWorld					Allow reloading of the world
      * @return true if succeeded.
      */
 	UFUNCTION(BlueprintCallable, Category = "Editor Scripting | Editor Revision Control Helpers", meta = (Keywords = "Source Control"))
-	static bool RevertAndReloadPackages(const TArray<FString>& InFilenames);
+	static bool RevertAndReloadPackages(const TArray<FString>& InPackagesToRevert, const TArray<FString>& InPackagesToReload, bool bRevertAll = false, bool bAllowReloadWorld = false);
+
+	/**
+	 * Reverts the provided files and then reloads all packages.
+	 * @param	InFilenames		The files/packages to revert
+	 * @return true if succeeded.
+	 */
+	static bool RevertAndReloadAllPackages(const TArray<FString>& InFilenames);
+
+	/**
+	 * Reverts all changes and then reloads all packages.
+	 * @return true if succeeded.
+	 */
+	static bool RevertAllChangesAndReloadAllPackages();
 
 #endif // !WITH_EDITOR
 
@@ -612,6 +631,18 @@ public:
 	 */
 	static bool GetAssetDataFromFileHistory(FSourceControlStatePtr InSourceControlState, TArray<FAssetData>& OutAssets, TArray<FName>* OutDependencies = nullptr, int64 MaxFetchSize = -1);
 
+	/**
+	 * Find packages in Source Control locations
+	 */
+	static bool ListAllPackages(TArray<FString>& OutPackageNames);
+
+	/**
+	 * Get the list of files and directories that source control should check when looking for changes.
+	 *
+	 * @param	bContentOnly	True to only include content directories.
+	 */
+	static TArray<FString> GetSourceControlLocations(const bool bContentOnly = false);
+
 };  // USourceControlHelpers
 
 
diff --git a/Engine/Source/Developer/SourceControl/Public/SourceControlOperationBase.h b/Engine/Source/Developer/SourceControl/Public/SourceControlOperationBase.h
--- a/Engine/Source/Developer/SourceControl/Public/SourceControlOperationBase.h
+++ b/Engine/Source/Developer/SourceControl/Public/SourceControlOperationBase.h
@@ -26,6 +26,12 @@ public:
 		ResultInfo.ErrorMessages.Add(InError);
 	}
 
+	/** Add tag. */
+	virtual void AddTag(const FString& InTag)
+	{
+		ResultInfo.Tags.Add(InTag);
+	}
+
 	/**
 	 * Append any info or error messages that may have accumulated during the operation prior
 	 * to returning a result, ensuring to keep any already accumulated info.
diff --git a/Engine/Source/Developer/SourceControl/Public/SourceControlOperations.h b/Engine/Source/Developer/SourceControl/Public/SourceControlOperations.h
--- a/Engine/Source/Developer/SourceControl/Public/SourceControlOperations.h
+++ b/Engine/Source/Developer/SourceControl/Public/SourceControlOperations.h
@@ -251,6 +251,16 @@ public:
 		return bIsSoftRevert;
 	}
 
+	void SetRevertAll(const bool bInRevertAll)
+	{
+		bIsRevertAll = bInRevertAll;
+	}
+
+	bool IsRevertAll() const
+	{
+		return bIsRevertAll;
+	}
+
 	bool ShouldDeleteNewFiles() const
 	{
 		return USourceControlPreferences::ShouldDeleteNewFilesOnRevert();
@@ -265,6 +275,7 @@ public:
 
 protected:
 	bool				bIsSoftRevert = false;
+	bool				bIsRevertAll = false;
 	TArray<FString>		DeletedFiles;
 };
 
diff --git a/Engine/Source/Developer/UncontrolledChangelists/Private/UncontrolledChangelistsModule.cpp b/Engine/Source/Developer/UncontrolledChangelists/Private/UncontrolledChangelistsModule.cpp
--- a/Engine/Source/Developer/UncontrolledChangelists/Private/UncontrolledChangelistsModule.cpp
+++ b/Engine/Source/Developer/UncontrolledChangelists/Private/UncontrolledChangelistsModule.cpp
@@ -366,7 +366,7 @@ bool FUncontrolledChangelistsModule::OnRevert(const TArray<FString>& InFilenames
 		return true;
 	}
 
-	bSuccess = SourceControlHelpers::ApplyOperationAndReloadPackages(InFilenames, ExecuteRevertOperation);
+	bSuccess = SourceControlHelpers::ApplyOperationAndReloadPackages(InFilenames, InFilenames, ExecuteRevertOperation);
 
 	UpdateStatus();
 
diff --git a/Engine/Source/Editor/ContentBrowser/Private/FrontendFilters.cpp b/Engine/Source/Editor/ContentBrowser/Private/FrontendFilters.cpp
--- a/Engine/Source/Editor/ContentBrowser/Private/FrontendFilters.cpp
+++ b/Engine/Source/Editor/ContentBrowser/Private/FrontendFilters.cpp
@@ -43,7 +43,7 @@
 #include "ProfilingDebugging/CpuProfilerTrace.h"
 #include "Settings/ContentBrowserSettings.h"
 #include "SourceControlOperations.h"
-#include "SourceControlWindows.h"
+#include "SourceControlHelpers.h"
 #include "Templates/RemoveReference.h"
 #include "Templates/UnrealTemplate.h"
 #include "TextFilterKeyValueHandlers.h"
@@ -857,7 +857,7 @@ void FFrontendFilter_NotSourceControlled::RequestStatus()
 		}
 		else
 		{
-			Filenames = FSourceControlWindows::GetSourceControlLocations(/*bContentOnly*/true);
+			Filenames = SourceControlHelpers::GetSourceControlLocations(/*bContentOnly*/true);
 		}
 		UpdateStatusOperation->SetCheckingAllFiles(false);
 		SourceControlProvider.Execute(UpdateStatusOperation, Filenames, EConcurrency::Asynchronous, FSourceControlOperationComplete::CreateSP(this, &FFrontendFilter_NotSourceControlled::SourceControlOperationComplete));
diff --git a/Engine/Source/Editor/SceneOutliner/Private/SceneOutlinerSCCHandler.cpp b/Engine/Source/Editor/SceneOutliner/Private/SceneOutlinerSCCHandler.cpp
--- a/Engine/Source/Editor/SceneOutliner/Private/SceneOutlinerSCCHandler.cpp
+++ b/Engine/Source/Editor/SceneOutliner/Private/SceneOutlinerSCCHandler.cpp
@@ -84,6 +84,7 @@ void FSceneOutlinerSCCHandler::CacheCanExecuteVars()
 	bCanExecuteSCC = true;
 	bCanExecuteSCCCheckOut = false;
 	bCanExecuteSCCCheckIn = false;
+	bCanExecuteSCCRevert = false;
 	bCanExecuteSCCHistory = false;
 	bUsesFileRevisions = false;
 	bUsesChangelists = false;
@@ -128,6 +129,11 @@ void FSceneOutlinerSCCHandler::CacheCanExecuteVars()
 				{
 					bCanExecuteSCCCheckIn = true;
 				}
+
+				if (SourceControlState->CanRevert())
+				{
+					bCanExecuteSCCRevert = true;
+				}
 			}
 		}
 
@@ -155,6 +161,11 @@ bool FSceneOutlinerSCCHandler::CanExecuteSCCCheckIn() const
 	return bCanExecuteSCCCheckIn && bUsesFileRevisions;
 }
 
+bool FSceneOutlinerSCCHandler::CanExecuteSCCRevert() const
+{
+	return bCanExecuteSCCRevert;
+}
+
 bool FSceneOutlinerSCCHandler::CanExecuteSCCHistory() const
 {
 	return bCanExecuteSCCHistory;
@@ -227,6 +238,20 @@ void FSceneOutlinerSCCHandler::FillSourceControlSubMenu(UToolMenu* Menu)
 		);
 	}
 
+	if (CanExecuteSCCRevert())
+	{
+		Section.AddMenuEntry(
+			"SCCRevert",
+			LOCTEXT("SCCRevert", "Revert"),
+			LOCTEXT("SCCRevertTooltip", "Reverts the item to the state it was before it was checked out."),
+			FSlateIcon(FAppStyle::GetAppStyleSetName(), "SourceControl.Actions.Revert"),
+			FUIAction(
+				FExecuteAction::CreateSP(this, &FSceneOutlinerSCCHandler::ExecuteSCCRevert),
+				FCanExecuteAction::CreateSP(this, &FSceneOutlinerSCCHandler::CanExecuteSCCRevert)
+			)
+		);
+	}
+
 	if (CanExecuteSCCShowInChangelist())
 	{
 		Section.AddSeparator(NAME_None);
@@ -339,6 +364,13 @@ void FSceneOutlinerSCCHandler::ExecuteSCCCheckIn()
 	}
 }
 
+void FSceneOutlinerSCCHandler::ExecuteSCCRevert()
+{
+	TArray<FString> PackageNames;
+	GetSelectedPackageNames(PackageNames);
+	FSourceControlWindows::PromptForRevert(PackageNames, true);
+}
+
 void FSceneOutlinerSCCHandler::ExecuteSCCHistory()
 {
 	TArray<FString> PackageNames;
diff --git a/Engine/Source/Editor/SceneOutliner/Public/SceneOutlinerSCCHandler.h b/Engine/Source/Editor/SceneOutliner/Public/SceneOutlinerSCCHandler.h
--- a/Engine/Source/Editor/SceneOutliner/Public/SceneOutlinerSCCHandler.h
+++ b/Engine/Source/Editor/SceneOutliner/Public/SceneOutlinerSCCHandler.h
@@ -22,6 +22,7 @@ private:
 	bool CanExecuteSCC() const;
 	bool CanExecuteSCCCheckOut() const;
 	bool CanExecuteSCCCheckIn() const;
+	bool CanExecuteSCCRevert() const;
 	bool CanExecuteSCCHistory() const;
 	bool CanExecuteSCCRefresh() const;
 	bool CanExecuteSCCShowInChangelist() const;
@@ -31,6 +32,7 @@ private:
 	void ExecuteSCCRefresh();
 	void ExecuteSCCCheckOut();
 	void ExecuteSCCCheckIn();
+	void ExecuteSCCRevert();
 	void ExecuteSCCHistory();
 	void ExecuteSCCShowInChangelist();
 
@@ -44,6 +46,7 @@ private:
 	bool bCanExecuteSCCCheckOut = false;
 	bool bCanExecuteSCCCheckIn = false;
 	bool bCanExecuteSCCHistory = false;
+	bool bCanExecuteSCCRevert = false;
 	bool bUsesFileRevisions = false;
 	bool bUsesChangelists = false;
 
diff --git a/Engine/Source/Editor/SourceControlWindows/Private/SSourceControlChangelists.cpp b/Engine/Source/Editor/SourceControlWindows/Private/SSourceControlChangelists.cpp
--- a/Engine/Source/Editor/SourceControlWindows/Private/SSourceControlChangelists.cpp
+++ b/Engine/Source/Editor/SourceControlWindows/Private/SSourceControlChangelists.cpp
@@ -1967,7 +1967,7 @@ void SSourceControlChangelistsWidget::OnRevert()
 
 			SSourceControlCommon::ExecuteChangelistOperationWithSlowTaskWrapper(LOCTEXT("Reverting_Files", "Reverting file(s)..."), [&PackageFiles, &NonPackageFiles]()
 			{
-				if ((PackageFiles.Num() == 0 || SourceControlHelpers::RevertAndReloadPackages(PackageFiles)) && (NonPackageFiles.Num() == 0 || SourceControlHelpers::RevertFiles(NonPackageFiles)))
+				if ((PackageFiles.Num() == 0 || SourceControlHelpers::RevertAndReloadPackages(PackageFiles, PackageFiles)) && (NonPackageFiles.Num() == 0 || SourceControlHelpers::RevertFiles(NonPackageFiles)))
 				{
 					SSourceControlCommon::DisplaySourceControlOperationNotification(LOCTEXT("Revert_Files_Succeeded", "The selected file(s) were reverted."), SNotificationItem::CS_Success);
 				}
diff --git a/Engine/Source/Editor/SourceControlWindows/Private/SSourceControlRevert.cpp b/Engine/Source/Editor/SourceControlWindows/Private/SSourceControlRevert.cpp
--- a/Engine/Source/Editor/SourceControlWindows/Private/SSourceControlRevert.cpp
+++ b/Engine/Source/Editor/SourceControlWindows/Private/SSourceControlRevert.cpp
@@ -397,7 +397,7 @@ private:
 	bool bRevertUnchangedFilesOnly;
 };
 
-bool FSourceControlWindows::PromptForRevert( const TArray<FString>& InPackageNames )
+bool FSourceControlWindows::PromptForRevert( const TArray<FString>& InPackageNames, bool InReloadAllPackages)
 {
 	bool bReverted = false;
 
@@ -440,7 +440,14 @@ bool FSourceControlWindows::PromptForRevert( const TArray<FString>& InPackageNam
 			
 			if (FinalPackagesToRevert.Num() > 0)
 			{
-				SourceControlHelpers::RevertAndReloadPackages(FinalPackagesToRevert);
+				if (InReloadAllPackages)
+				{
+					SourceControlHelpers::RevertAndReloadAllPackages(FinalPackagesToRevert);
+				}
+				else
+				{
+					SourceControlHelpers::RevertAndReloadPackages(FinalPackagesToRevert, FinalPackagesToRevert, /*bRevertAll=*/false, /*bAllowReloadWorld=*/true);
+				}
 
 				bReverted = true;
 			}
@@ -450,4 +457,9 @@ bool FSourceControlWindows::PromptForRevert( const TArray<FString>& InPackageNam
 	return bReverted;
 }
 
+bool FSourceControlWindows::RevertAllChangesAndReloadAllPackages()
+{	
+	return SourceControlHelpers::RevertAllChangesAndReloadAllPackages();
+}
+
 #undef LOCTEXT_NAMESPACE
diff --git a/Engine/Source/Editor/SourceControlWindows/Private/SourceControlWindows.cpp b/Engine/Source/Editor/SourceControlWindows/Private/SourceControlWindows.cpp
--- a/Engine/Source/Editor/SourceControlWindows/Private/SourceControlWindows.cpp
+++ b/Engine/Source/Editor/SourceControlWindows/Private/SourceControlWindows.cpp
@@ -36,29 +36,6 @@ FCheckinResultInfo::FCheckinResultInfo()
 
 TWeakPtr<SNotificationItem> FSourceControlWindows::ChoosePackagesToCheckInNotification;
 
-TArray<FString> FSourceControlWindows::GetSourceControlLocations(const bool bContentOnly)
-{
-	TArray<FString> SourceControlLocations;
-
-	{
-		TArray<FString> RootPaths;
-		FPackageName::QueryRootContentPaths(RootPaths);
-		for (const FString& RootPath : RootPaths)
-		{
-			const FString RootPathOnDisk = FPackageName::LongPackageNameToFilename(RootPath);
-			SourceControlLocations.Add(FPaths::ConvertRelativePathToFull(RootPathOnDisk));
-		}
-	}
-
-	if (!bContentOnly)
-	{
-		SourceControlLocations.Add(FPaths::ConvertRelativePathToFull(FPaths::ProjectConfigDir()));
-		SourceControlLocations.Add(FPaths::ConvertRelativePathToFull(FPaths::GetProjectFilePath()));
-	}
-
-	return SourceControlLocations;
-}
-
 bool FSourceControlWindows::ChoosePackagesToCheckIn(const FSourceControlWindowsOnCheckInComplete& OnCompleteDelegate)
 {
 	if (!ISourceControlModule::Get().IsEnabled())
@@ -96,7 +73,7 @@ bool FSourceControlWindows::ChoosePackagesToCheckIn(const FSourceControlWindowsO
 	}
 	else
 	{
-		Filenames = GetSourceControlLocations();
+		Filenames = SourceControlHelpers::GetSourceControlLocations();
 	}
 	
 	// make sure the SourceControlProvider state cache is populated as well
@@ -201,48 +178,6 @@ static bool SaveDirtyPackages()
 	return bSaved;
 }
 
-
-static bool ListAllPackages(TArray<FString>& OutPackageNames)
-{
-	// determine source control locations
-	TArray<FString> SourceControlLocations;
-
-	if (ISourceControlModule::Get().UsesCustomProjectDir())
-	{
-		SourceControlLocations.Add(ISourceControlModule::Get().GetSourceControlProjectDir());
-	}
-	else
-	{
-		FSourceControlWindows::GetSourceControlLocations(/*bContentOnly=*/true);
-	}
-
-	// find packages in those locations
-	for (const FString& SourceControlLocation : SourceControlLocations)
-	{
-		TArray<FString> PackageRelativePaths;
-		if (FPackageName::FindPackagesInDirectory(PackageRelativePaths, FPaths::ConvertRelativePathToFull(SourceControlLocation)))
-		{
-			OutPackageNames.Reserve(OutPackageNames.Num() + PackageRelativePaths.Num());
-			for (const FString& Path : PackageRelativePaths)
-			{
-				FString PackageName;
-				FString FailureReason;
-				if (FPackageName::TryConvertFilenameToLongPackageName(Path, PackageName, &FailureReason))
-				{
-					OutPackageNames.Add(PackageName);
-				}
-				else
-				{
-					FMessageLog("SourceControl").Error(FText::FromString(FailureReason));
-				}
-			}
-		}
-	}
-
-	return true;
-}
-
-
 bool FSourceControlWindows::SyncAllPackages()
 {
 	bool bSaved = SaveDirtyPackages();
@@ -251,7 +186,7 @@ bool FSourceControlWindows::SyncAllPackages()
 	if (bSaved)
 	{
 		TArray<FString> PackageNames;
-		if (ListAllPackages(PackageNames))
+		if (SourceControlHelpers::ListAllPackages(PackageNames))
 		{
 			// sync those packages if any were found
 			if (PackageNames.Num() > 0)
@@ -586,7 +521,7 @@ void FSourceControlWindows::ChoosePackagesToCheckInCompleted(const TArray<UPacka
 	}
 	else
 	{
-		PendingDeletePaths = GetSourceControlLocations();
+		PendingDeletePaths = SourceControlHelpers::GetSourceControlLocations();
 	}
 
 	PromptForCheckin(OutResultInfo, PackageNames, PendingDeletePaths, ConfigFiles, bUseSourceControlStateCache);
diff --git a/Engine/Source/Editor/SourceControlWindows/Public/SourceControlWindows.h b/Engine/Source/Editor/SourceControlWindows/Public/SourceControlWindows.h
--- a/Engine/Source/Editor/SourceControlWindows/Public/SourceControlWindows.h
+++ b/Engine/Source/Editor/SourceControlWindows/Public/SourceControlWindows.h
@@ -35,13 +35,6 @@ DECLARE_DELEGATE_OneParam(FSourceControlWindowsOnCheckInComplete, const FCheckin
 class SOURCECONTROLWINDOWS_API FSourceControlWindows
 {
 public:
-	/**
-	 * Get the list of files and directories that source control should check when looking for changes.
-	 *
-	 * @param	bContentOnly	True to only include content directories.
-	 */
-	static TArray<FString> GetSourceControlLocations(const bool bContentOnly = false);
-
 	/**
 	 * Opens a user dialog to choose packages to submit.
 	 *
@@ -102,10 +95,16 @@ public:
 	 * Prompt the user with a revert files dialog, allowing them to specify which packages, if any, should be reverted.
 	 *
 	 * @param	InPackageNames	Names of the packages to consider for reverting
+	 * @param	InReloadAllPackages	Reload all packages as part of the revert operation
 	 *
 	 * @return	true if the files were reverted; false if the user canceled out of the dialog
 	 */
-	static bool PromptForRevert(const TArray<FString>& InPackageNames );
+	static bool PromptForRevert(const TArray<FString>& InPackageNames, bool InReloadAllPackages = false );
+
+	/**
+	 * Revert all locally modified files, and reload all packages
+	 */
+	static bool RevertAllChangesAndReloadAllPackages();
 
 	/**
 	 * Displays file diff against workspace version
diff --git a/Engine/Source/Editor/StatusBar/Private/SourceControlMenuHelpers.cpp b/Engine/Source/Editor/StatusBar/Private/SourceControlMenuHelpers.cpp
--- a/Engine/Source/Editor/StatusBar/Private/SourceControlMenuHelpers.cpp
+++ b/Engine/Source/Editor/StatusBar/Private/SourceControlMenuHelpers.cpp
@@ -17,6 +17,7 @@
 #include "Widgets/Input/SComboButton.h"
 #include "Widgets/Input/SButton.h"
 #include "Widgets/Images/SImage.h"
+#include "Widgets/Layout/SSeparator.h"
 #include "Widgets/Images/SLayeredImage.h"
 #include "LevelEditorActions.h"
 #include "PackageTools.h"
@@ -48,6 +49,7 @@ void FSourceControlCommands::RegisterCommands()
 	UI_COMMAND(ViewChangelists, "View Changelists", "Opens a dialog displaying current changelists.", EUserInterfaceActionType::Button, FInputChord());
 	UI_COMMAND(SubmitContent, "Submit Content", "Opens a dialog with check in options for content and levels.", EUserInterfaceActionType::Button, FInputChord());
 	UI_COMMAND(CheckOutModifiedFiles, "Check Out Modified Files", "Opens a dialog to check out any assets which have been modified.", EUserInterfaceActionType::Button, FInputChord());
+	UI_COMMAND(RevertAll, "Revert All Files", "Opens a dialog to revert any assets which have been modified.", EUserInterfaceActionType::Button, FInputChord());
 
 	ActionList->MapAction(
 		ConnectToSourceControl,
@@ -80,6 +82,12 @@ void FSourceControlCommands::RegisterCommands()
 		FExecuteAction::CreateStatic(&FSourceControlCommands::CheckOutModifiedFiles_Clicked),
 		FCanExecuteAction::CreateStatic(&FSourceControlCommands::CheckOutModifiedFiles_CanExecute)
 	);
+
+	ActionList->MapAction(
+		RevertAll,
+		FExecuteAction::CreateStatic(&FSourceControlCommands::RevertAllModifiedFiles_Clicked),
+		FCanExecuteAction::CreateLambda([] { return true; })
+	);
 }
 
 void FSourceControlCommands::ConnectToSourceControl_Clicked()
@@ -146,6 +154,18 @@ void FSourceControlCommands::CheckOutModifiedFiles_Clicked()
 	FEditorFileUtils::PromptForCheckoutAndSave(PackagesToSave, bCheckDirty, bPromptUserToSave);
 }
 
+void FSourceControlCommands::RevertAllModifiedFiles_Clicked()
+{
+	FText Message = LOCTEXT("RevertAllModifiedFiles", "Are you sure you want to revert all local changes? By proceeding, your local changes will be discarded, and the state of your last synced snapshot will be restored.");
+	FText Title = LOCTEXT("RevertAllModifiedFiles_Title", "Revert all local changes");
+	if (FMessageDialog::Open(EAppMsgType::YesNo, EAppReturnType::No, Message, &Title) == EAppReturnType::No)
+	{
+		return;
+	}
+	
+	FSourceControlWindows::RevertAllChangesAndReloadAllPackages();
+}
+
 FSourceControlMenuHelpers& FSourceControlMenuHelpers::Get()
 {
 	// Singleton instance
@@ -225,6 +245,29 @@ TSharedRef<SWidget> FSourceControlMenuHelpers::GenerateSourceControlMenuContent(
 	return UToolMenus::Get()->GenerateWidget("StatusBar.ToolBar.SourceControl", FToolMenuContext(FSourceControlCommands::ActionList));
 }
 
+TSharedRef<SWidget> FSourceControlMenuHelpers::GenerateCheckInComboButtonContent()
+{
+	UToolMenu* SourceControlMenu = UToolMenus::Get()->RegisterMenu("StatusBar.ToolBar.SourceControl.CheckInCombo", NAME_None, EMultiBoxType::Menu, false);
+
+	FToolMenuSection& Section = SourceControlMenu->AddSection("SourceControlComboActions", LOCTEXT("SourceControlComboMenuHeadingActions", "Actions"));
+
+	Section.AddMenuEntry(
+		FSourceControlCommands::Get().RevertAll,
+		TAttribute<FText>(),
+		TAttribute<FText>(),
+		FSlateIcon(FAppStyle::GetAppStyleSetName(), "SourceControl.Actions.Revert")
+	);
+
+	Section.AddMenuEntry(
+		FSourceControlCommands::Get().ViewChangelists,
+		TAttribute<FText>(),
+		TAttribute<FText>(),
+		FSlateIcon(FAppStyle::GetAppStyleSetName(), "SourceControl.ChangelistsTab")
+	);
+
+	return UToolMenus::Get()->GenerateWidget("StatusBar.ToolBar.SourceControl.CheckInCombo", FToolMenuContext(FSourceControlCommands::ActionList));
+}
+
 FText FSourceControlMenuHelpers::GetSourceControlStatusText()
 {
 	if (QueryState == EQueryState::Querying)
@@ -474,6 +517,25 @@ TSharedRef<SWidget> FSourceControlMenuHelpers::MakeSourceControlStatusWidget()
 			]
 			.OnClicked_Static(&FSourceControlMenuHelpers::OnSourceControlCheckInChangesClicked)
 		]
+		+ SHorizontalBox::Slot()
+		.VAlign(VAlign_Center)
+		.AutoWidth()
+		[
+			SNew(SSeparator)
+			.Thickness(1.0)
+			.Orientation(EOrientation::Orient_Vertical)
+		]
+		+SHorizontalBox::Slot() // Check In Kebab Combo button
+		.VAlign(VAlign_Center)
+		.AutoWidth()
+		[
+			SNew(SComboButton)
+			.ContentPadding(FMargin(7.f, 0.f))
+			.ComboButtonStyle(&FAppStyle::Get().GetWidgetStyle<FComboButtonStyle>("StatusBar.StatusBarEllipsisComboButton"))
+			.MenuPlacement(MenuPlacement_AboveAnchor)
+			.Visibility_Static(&FSourceControlMenuHelpers::GetSourceControlCheckInStatusVisibility)
+			.OnGetMenuContent(FOnGetContent::CreateStatic(&FSourceControlMenuHelpers::GenerateCheckInComboButtonContent))
+		]
 		+ SHorizontalBox::Slot() // Sync Latest Button
 		.VAlign(VAlign_Center)
 		.AutoWidth()
diff --git a/Engine/Source/Editor/StatusBar/Private/SourceControlMenuHelpers.h b/Engine/Source/Editor/StatusBar/Private/SourceControlMenuHelpers.h
--- a/Engine/Source/Editor/StatusBar/Private/SourceControlMenuHelpers.h
+++ b/Engine/Source/Editor/StatusBar/Private/SourceControlMenuHelpers.h
@@ -29,6 +29,7 @@ private:
 	static void ViewChangelists_Clicked();
 	static bool CheckOutModifiedFiles_CanExecute();
 	static void CheckOutModifiedFiles_Clicked();
+	static void RevertAllModifiedFiles_Clicked();
 
 public:
 	/**
@@ -39,6 +40,7 @@ public:
 	TSharedPtr< FUICommandInfo > ViewChangelists;
 	TSharedPtr< FUICommandInfo > SubmitContent;
 	TSharedPtr< FUICommandInfo > CheckOutModifiedFiles;
+	TSharedPtr< FUICommandInfo > RevertAll;
 
 	static TSharedRef<FUICommandList> ActionList;
 };
@@ -72,6 +74,7 @@ public:
 private:
 	static void OnSourceControlOperationComplete(const FSourceControlOperationRef& InOperation, ECommandResult::Type InResult);
 	static TSharedRef<SWidget> GenerateSourceControlMenuContent();
+	static TSharedRef<SWidget> GenerateCheckInComboButtonContent();
 	static FText GetSourceControlStatusText();
 	static FText GetSourceControlTooltip();
 	static const FSlateBrush* GetSourceControlIconBadge();
diff --git a/Engine/Source/Runtime/SlateCore/Private/Styling/StarshipCoreStyle.cpp b/Engine/Source/Runtime/SlateCore/Private/Styling/StarshipCoreStyle.cpp
--- a/Engine/Source/Runtime/SlateCore/Private/Styling/StarshipCoreStyle.cpp
+++ b/Engine/Source/Runtime/SlateCore/Private/Styling/StarshipCoreStyle.cpp
@@ -1830,8 +1830,15 @@ void FStarshipCoreStyle::SetupDockingStyles(TSharedRef<FStyle>& Style)
 		.SetDownArrowPadding(FMargin(0.0f))
 		.SetDownArrowAlignment(EVerticalAlignment::VAlign_Bottom);
 
+	const FComboButtonStyle StatusBarEllipsisComboButton = FComboButtonStyle(Style->GetWidgetStyle<FComboButtonStyle>("ComboButton"))
+		.SetDownArrowImage(IMAGE_BRUSH_SVG("Starship/Common/ellipsis-vertical-narrow", FVector2D(6, 24)))
+		.SetButtonStyle(StatusBarButton)
+		.SetDownArrowPadding(FMargin(0.0f))
+		.SetDownArrowAlignment(EVerticalAlignment::VAlign_Bottom);
+
 	Style->Set("StatusBar.StatusBarButton", StatusBarButton);
 	Style->Set("StatusBar.StatusBarComboButton", StatusBarComboButton);
+	Style->Set("StatusBar.StatusBarEllipsisComboButton", StatusBarEllipsisComboButton);
 	Style->Set("StatusBar.Height", 32.0f);
 	Style->Set("StatusBar.DrawerShadow", new BOX_BRUSH("Starship/StatusBar/drawer-shadow-bottom", FMargin(10 / 64.0f, 20 / 64.f, 10 / 64.f, 0), FLinearColor(0, 0, 0, 1)));
 	Style->Set("StatusBar.DrawerBackground", new FSlateRoundedBoxBrush(FStyleColors::Panel, 5.0f));
