* Refactored Sync Latest operation to improve performance
* Sync Latest no longer listing and reloading all packages when the source control provider provides a sync preview functionality. When a Sync Preview is available, Sync Latest unloads only the files that are going to be overwritten by the sync operation, and then reloads them and the world

* Refactored Revert All operation to improve performance
* Revert All no longer lists and reloads all packages when the source control provider provides a sync preview functionality. Instead it unloads only the files that may be overwritten by the revert operation (i.e. submittable files), and then reloads them and the world
* Replaced RevertAndReloadAllPackages with RevertAndReloadWorld, no longer listing and reloading all packages (see https://p4-swarm.epicgames.net/reviews/23244245/)

This review is a combination of the two archived reviews https://p4-swarm.epicgames.net/reviews/23736137 and https://p4-swarm.epicgames.net/reviews/23736282

diff --git a/Engine/Plugins/Editor/AssetManagerEditor/Source/AssetManagerEditor/Private/AssetSourceControlContextMenu.cpp b/Engine/Plugins/Editor/AssetManagerEditor/Source/AssetManagerEditor/Private/AssetSourceControlContextMenu.cpp
--- a/Engine/Plugins/Editor/AssetManagerEditor/Source/AssetManagerEditor/Private/AssetSourceControlContextMenu.cpp
+++ b/Engine/Plugins/Editor/AssetManagerEditor/Source/AssetManagerEditor/Private/AssetSourceControlContextMenu.cpp
@@ -773,21 +773,19 @@ void FAssetSourceControlContextMenuState::ExecuteSCCRevert() const
 {
 	TArray<FString> PackageNames;
 	GetSelectedPackageNames(PackageNames);
-
-	bool bReloadAll = false;
+	bool bReloadWorld = false;
 	for (const FString& PackageName : PackageNames)
 	{
 		if (UPackage* Package = FindPackage(nullptr, *PackageName))
 		{
 			if (Package->ContainsMap())
 			{
-				bReloadAll = true;
+				bReloadWorld = true;
 				break;
 			}
 		}
 	}
-
-	FSourceControlWindows::PromptForRevert(PackageNames, bReloadAll);
+	FSourceControlWindows::PromptForRevert(PackageNames, bReloadWorld);
 }
 
 void FAssetSourceControlContextMenuState::ExecuteSCCRevertWritable() const
@@ -848,7 +846,7 @@ void FAssetSourceControlContextMenuState::ExecuteSCCRevertWritable() const
 		return true;
 	};
 
-	SourceControlHelpers::ApplyOperationAndReloadPackages(PackageNames, PackageNames, RevertOperation);
+	SourceControlHelpers::ApplyOperationAndReloadPackages(PackageNames, RevertOperation);
 
 	// Tell UncontrolledCL to refresh to pick up changes to files it was watching.
 	FUncontrolledChangelistsModule& UncontrolledChangelistsModule = FUncontrolledChangelistsModule::Get();
diff --git a/Engine/Source/Developer/AssetTools/Private/AssetViewUtils.cpp b/Engine/Source/Developer/AssetTools/Private/AssetViewUtils.cpp
--- a/Engine/Source/Developer/AssetTools/Private/AssetViewUtils.cpp
+++ b/Engine/Source/Developer/AssetTools/Private/AssetViewUtils.cpp
@@ -1449,7 +1449,7 @@ FString AssetViewUtils::GetPackagePathWithinRoot(const FString& PackageName)
 	return RelativePathToAsset;
 }
 
-int32 AssetViewUtils::GetPackageLengthForCooking(const FString& PackageName, bool IsInternalBuild)
+int32 AssetViewUtils::GetPackageLengthForCooking(const FString& PackageName, bool bIsInternalBuild)
 {
 	FString RelativePathToAsset;
 	if (!FPackageName::TryConvertLongPackageNameToFilename(PackageName, RelativePathToAsset, FPackageName::GetAssetPackageExtension()))
@@ -1485,13 +1485,13 @@ int32 AssetViewUtils::GetPackageLengthForCooking(const FString& PackageName, boo
 
 		if (FString ExternalPluginCookedRootPath = CVarExternalPluginCookedRootPath->GetValueOnGameThread(); ExternalPluginCookedRootPath.Len() > 0)
 		{
-			IsInternalBuild = false;
+			bIsInternalBuild = false;
 			AbsoluteTargetPath = MoveTemp(ExternalPluginCookedRootPath);
 			AbsoluteTargetPath /= FString();
 		}
 	}
 
-	if (IsInternalBuild)
+	if (bIsInternalBuild)
 	{
 		// We assume a constant size for the build machine base path for things that reside within the UE source tree
 		const FString AbsoluteUERootPath = FPaths::ConvertRelativePathToFull(FPaths::RootDir());
@@ -1704,13 +1704,15 @@ void ShowSyncDependenciesDialog(const TArray<FString>& InDependencies, TArray<FS
 	}
 }
 
-void AssetViewUtils::SyncPackagesFromSourceControl(const TArray<FString>& PackageNames)
+void AssetViewUtils::SyncPackagesFromSourceControl(const TArray<FString>& PackageNames, bool bIsSyncLatestOperation)
 {
 	if (PackageNames.Num() > 0)
 	{
-		// Warn about any packages that are being synced without also getting the newest version of their dependencies...
 		TArray<FString> PackageNamesToSync = PackageNames;
+
+		if (!bIsSyncLatestOperation)
 		{
+			// Warn about any packages that are being synced without also getting the newest version of their dependencies...
 			TArray<FString> OutOfDateDependencies;
 			GetOutOfDatePackageDependencies(PackageNamesToSync, OutOfDateDependencies);
 
@@ -1744,7 +1746,20 @@ void AssetViewUtils::SyncPackagesFromSourceControl(const TArray<FString>& Packag
 		}
 
 		// Sync everything...
-		SCCProvider.Execute(ISourceControlOperation::Create<FSync>(), PackageFilenames);
+		TSharedRef<FSync> Operation = ISourceControlOperation::Create<FSync>();
+		if (bIsSyncLatestOperation)
+		{
+			Operation->SetHeadRevisionFlag(true);
+			if (SCCProvider.Execute(Operation) != ECommandResult::Succeeded)
+			{
+				UE_LOG(LogAssetViewTools, Warning, TEXT("Error while syncing latest revision"));
+				return;
+			}
+		}
+		else
+		{
+			SCCProvider.Execute(Operation, PackageFilenames);
+		}
 
 		// Syncing may have deleted some packages, so we need to unload those rather than re-load them...
 		// Note: we will store the package using weak pointers here otherwise we might have garbage collection issues after the ReloadPackages call
@@ -1774,6 +1789,11 @@ void AssetViewUtils::SyncPackagesFromSourceControl(const TArray<FString>& Packag
 		}
 
 		// Hot-reload the new packages...
+		if (bIsSyncLatestOperation)
+		{
+			UWorld* EditorWorld = GEditor->GetEditorWorldContext().World();
+			LoadedPackages.Add(EditorWorld->GetPackage());
+		}
 		UPackageTools::ReloadPackages(LoadedPackages);
 
 		// Restore the camera views after a map reload if necessary...
@@ -1806,6 +1826,49 @@ void AssetViewUtils::SyncPackagesFromSourceControl(const TArray<FString>& Packag
 	}
 }
 
+void AssetViewUtils::SyncLatestFromSourceControl()
+{
+	TArray<FString> PackageNames;
+	TArray<FString> PackageFilenames;
+
+	ISourceControlProvider& SCCProvider = ISourceControlModule::Get().GetProvider();
+	TSharedRef<FSyncPreview> PreviewOperation = ISourceControlOperation::Create<FSyncPreview>();
+	PreviewOperation->SetHeadRevisionFlag(true);
+	ECommandResult::Type OperationResult =  SCCProvider.Execute(PreviewOperation);
+
+	if (OperationResult == ECommandResult::Succeeded)
+	{
+		PackageFilenames = PreviewOperation->GetAffectedFiles();
+
+		for (const FString& PackageFilename : PackageFilenames)
+		{
+			FString PackageName;
+			if (FPackageName::TryConvertFilenameToLongPackageName(PackageFilename, PackageName))
+			{
+				PackageNames.Add(PackageName);
+			}
+		}
+
+		SyncPackagesFromSourceControl(PackageNames, /*IsSyncLatestOperation=*/true);
+	}
+	else
+	{
+		// Likely the source control provider does not support sync previews. Proceed to syncing all paths under source control
+		TArray<FString> SourceControlLocations;
+
+		if (ISourceControlModule::Get().UsesCustomProjectDir())
+		{
+			SourceControlLocations.Add(ISourceControlModule::Get().GetSourceControlProjectDir());
+		}
+		else
+		{
+			SourceControlLocations = SourceControlHelpers::GetSourceControlLocations(/*bContentOnly=*/true);
+		}
+
+		SyncPathsFromSourceControl(SourceControlLocations);
+	}
+}
+
 void AssetViewUtils::SyncPathsFromSourceControl(const TArray<FString>& ContentPaths)
 {
 	TArray<FString> PathsOnDisk;
diff --git a/Engine/Source/Developer/AssetTools/Public/AssetViewUtils.h b/Engine/Source/Developer/AssetTools/Public/AssetViewUtils.h
--- a/Engine/Source/Developer/AssetTools/Public/AssetViewUtils.h
+++ b/Engine/Source/Developer/AssetTools/Public/AssetViewUtils.h
@@ -197,7 +197,7 @@ namespace AssetViewUtils
 	ASSETTOOLS_API FString GetPackagePathWithinRoot(const FString& PackageName);
 
 	/** Returns the length of the computed cooked package name and path whether it's run on a build machine or locally */
-	ASSETTOOLS_API int32 GetPackageLengthForCooking(const FString& PackageName, bool IsInternalBuild);
+	ASSETTOOLS_API int32 GetPackageLengthForCooking(const FString& PackageName, bool bIsInternalBuild);
 
 	/** Checks to see whether the path is within the size restrictions for cooking */
 	ASSETTOOLS_API bool IsValidPackageForCooking(const FString& PackageName, FText& OutErrorMessage);
@@ -208,12 +208,15 @@ namespace AssetViewUtils
 	/** Gets the maximum path length for a cooked file. Changes behavior based on whether the editor experimental setting for long paths is enabled. */
 	ASSETTOOLS_API int32 GetMaxCookPathLen();
 
-	/** Syncs the specified packages from source control, other than any level assets which are currently being edited */
-	ASSETTOOLS_API void SyncPackagesFromSourceControl(const TArray<FString>& PackageNames);
+	/** Syncs the specified packages from source control, other than any level assets which are currently being edited.*/
+	ASSETTOOLS_API void SyncPackagesFromSourceControl(const TArray<FString>& PackageNames, bool bIsSyncLatestOperation = false);
 
 	/** Syncs the content from the specified paths from source control, other than any level assets which are currently being edited */
 	ASSETTOOLS_API void SyncPathsFromSourceControl(const TArray<FString>& ContentPaths);
 
+	/** Syncs latest from source control and attempts to reload the world */
+	ASSETTOOLS_API void SyncLatestFromSourceControl();
+
 	/** Show an error notification toast if the given error message is not empty */
 	ASSETTOOLS_API void ShowErrorNotifcation(const FText& InErrorMsg);
 
diff --git a/Engine/Source/Developer/SourceControl/Private/SourceControlHelpers.cpp b/Engine/Source/Developer/SourceControl/Private/SourceControlHelpers.cpp
--- a/Engine/Source/Developer/SourceControl/Private/SourceControlHelpers.cpp
+++ b/Engine/Source/Developer/SourceControl/Private/SourceControlHelpers.cpp
@@ -24,6 +24,8 @@
 #include "Editor.h"
 #include "PackageTools.h"
 #include "ObjectTools.h"
+#include "FileHelpers.h"
+#include "LevelEditorViewport.h"
 #endif
 
 #define LOCTEXT_NAMESPACE "SourceControlHelpers"
@@ -907,14 +909,12 @@ bool USourceControlHelpers::RevertFile(const FString& InFile, bool bSilent)
 }
 
 #if WITH_EDITOR
-bool USourceControlHelpers::ApplyOperationAndReloadPackages(const TArray<FString>& InPackagesToApplyOperation, const TArray<FString>& InPackagesToReload, 
-	const TFunctionRef<bool(const TArray<FString>&)>& InOperation, bool bAllowReloadWorld)
+bool USourceControlHelpers::ApplyOperationAndReloadPackages(const TArray<FString>& InFilenames,	const TFunctionRef<bool(const TArray<FString>&)>& InOperation, bool bReloadWorld)
 {
 	ISourceControlProvider& SourceControlProvider = ISourceControlModule::Get().GetProvider();
 	TArray<UPackage*> LoadedPackages;
-	TArray<FString> PackageNamesToApplyOperation;
-	TArray<FString> PackageFilenamesToApplyOperation;
-	TArray<FString> PackageNamesToReload;
+	TArray<FString> PackageNames;
+	TArray<FString> PackageFilenames;
 	TArray<FString> FilteredActorPackages;
 	bool bSuccess = false;
 
@@ -929,56 +929,59 @@ bool USourceControlHelpers::ApplyOperationAndReloadPackages(const TArray<FString
 	};
 
 	// Normalize packagenames and filenames
-	for (const FString& Filename : InPackagesToApplyOperation)
+	for (const FString& Filename : InFilenames)
 	{
 		FString Result;
 		
 		if (FPackageName::TryConvertFilenameToLongPackageName(Filename, Result))
 		{
-			PackageNamesToApplyOperation.Add(MoveTemp(Result));
+			PackageNames.Add(MoveTemp(Result));
 		}
 		else
 		{
-			PackageNamesToApplyOperation.Add(Filename);
-		}
-	}
-
-	for (const FString& Filename : InPackagesToReload)
-	{
-		FString Result;
-
-		if (FPackageName::TryConvertFilenameToLongPackageName(Filename, Result))
-		{
-			PackageNamesToReload.Add(MoveTemp(Result));
-		}
-		else
-		{
-			PackageNamesToReload.Add(Filename);
+			PackageNames.Add(Filename);
 		}
 	}
 
 	// Remove packages if they are loaded actors or world
-	PackageNamesToReload.RemoveAll([&FilteredActorPackages, &LoadedPackages, bAllowReloadWorld, &DetachLinker](const FString& PackageName) -> bool
+	bool bWorldFound = false;
+	PackageNames.RemoveAll([&FilteredActorPackages, &LoadedPackages, bReloadWorld, &DetachLinker, &bWorldFound](const FString& PackageName) -> bool
 	{
 		UPackage* Package = FindPackage(NULL, *PackageName);
 		
 		if (Package != nullptr)
 		{
-			if (UWorld* World = UWorld::FindWorldInPackage(Package); World && !bAllowReloadWorld)
+			if (UWorld* World = UWorld::FindWorldInPackage(Package); World)
 			{
-				FilteredActorPackages.Emplace(PackageName);
-				return true; // remove the package
+				if (!bReloadWorld)
+				{
+					FilteredActorPackages.Emplace(PackageName);
+					return true; // remove the package
+				}
+				else
+				{
+					bWorldFound = true;
+				}
 			}
 			else if (AActor* Actor = AActor::FindActorInPackage(Package))
 			{
-				if (bAllowReloadWorld)
+				if (bReloadWorld)
 				{
 					DetachLinker(Package);
-					return true;
-				}
 
-				FilteredActorPackages.Emplace(PackageName);
-				return true; // remove the package
+					if (!bWorldFound && Actor->GetWorld())
+					{
+						bWorldFound = true;
+						LoadedPackages.Add(Actor->GetWorld()->GetPackage());
+					}
+
+					return false;
+				}
+				else if (Actor->IsPackageExternal())
+				{
+					FilteredActorPackages.Emplace(PackageName);
+					return true; // remove the package
+				}				
 			}
 
 			LoadedPackages.Add(Package);
@@ -1004,10 +1007,10 @@ bool USourceControlHelpers::ApplyOperationAndReloadPackages(const TArray<FString
 		DetachLinker(Package);
 	}
 
-	PackageFilenamesToApplyOperation = SourceControlHelpers::PackageFilenames(PackageNamesToApplyOperation);
+	PackageFilenames = SourceControlHelpers::PackageFilenames(PackageNames);
 
 	// Apply Operation
-	bSuccess = InOperation(PackageFilenamesToApplyOperation);
+	bSuccess = InOperation(PackageFilenames);
 
 	// Reverting may have deleted some packages, so we need to delete those and unload them rather than re-load them...
 	FAssetRegistryModule& AssetRegistryModule = FModuleManager::LoadModuleChecked<FAssetRegistryModule>(TEXT("AssetRegistry"));
@@ -1033,6 +1036,22 @@ bool USourceControlHelpers::ApplyOperationAndReloadPackages(const TArray<FString
 		return false; // keep package
 	});
 
+	struct FCameraView
+	{
+		FVector Location;
+		FRotator Rotation;
+	};
+	TMap<FLevelEditorViewportClient*, FCameraView> CameraViews;
+
+	if (bReloadWorld)
+	{
+		// As we're reloading the world, we retain the camera views so they can be restored after the reload.
+		for (FLevelEditorViewportClient* LevelVC : GEditor->GetLevelViewportClients())
+		{
+			CameraViews.Add(LevelVC, { LevelVC->GetViewLocation(), LevelVC->GetViewRotation() });
+		}
+	}
+
 	// Hot-reload the new packages...
 	FText OutReloadErrorMsg;
 	const bool bInteractive = true;
@@ -1042,6 +1061,22 @@ bool USourceControlHelpers::ApplyOperationAndReloadPackages(const TArray<FString
 		UE_LOG(LogSourceControl, Warning, TEXT("%s"), *OutReloadErrorMsg.ToString());
 	}
 
+	if (bReloadWorld)
+	{
+		// Restore the camera views after a world reload if necessary...
+		for (FLevelEditorViewportClient* LevelVC : GEditor->GetLevelViewportClients())
+		{
+			const FCameraView& CameraView = CameraViews.FindChecked(LevelVC);
+			LevelVC->SetViewLocation(CameraView.Location);
+			if (!LevelVC->IsOrtho())
+			{
+				LevelVC->SetViewRotation(CameraView.Rotation);
+			}
+			LevelVC->Invalidate();
+			FEditorDelegates::OnEditorCameraMoved.Broadcast(CameraView.Location, CameraView.Rotation, LevelVC->ViewportType, LevelVC->ViewIndex);
+		}
+	}
+
 	// Delete and Unload assets...
 	if(ObjectTools::DeleteObjectsUnchecked(ObjectsToDelete) != ObjectsToDelete.Num())
 	{ 
@@ -1049,7 +1084,7 @@ bool USourceControlHelpers::ApplyOperationAndReloadPackages(const TArray<FString
 	}
 	
 	// Re-cache the SCC state...
-	if (bAllowReloadWorld)
+	if (bReloadWorld)
 	{
 		auto UpdateCommand = ISourceControlOperation::Create<FUpdateStatus>();
 		UpdateCommand->SetCheckingAllFiles(true);
@@ -1058,7 +1093,7 @@ bool USourceControlHelpers::ApplyOperationAndReloadPackages(const TArray<FString
 	}
 	else
 	{
-		SourceControlProvider.Execute(ISourceControlOperation::Create<FUpdateStatus>(), PackageFilenamesToApplyOperation, EConcurrency::Asynchronous);
+		SourceControlProvider.Execute(ISourceControlOperation::Create<FUpdateStatus>(), PackageFilenames, EConcurrency::Asynchronous);
 	}
 
 	return bSuccess;
@@ -1087,74 +1122,56 @@ TArray<FString> USourceControlHelpers::GetSourceControlLocations(const bool bCon
 	return SourceControlLocations;
 }
 
-// TODO: Move outside of helpers
-bool USourceControlHelpers::ListAllPackages(TArray<FString>& OutPackageNames)
+bool USourceControlHelpers::ListRevertablePackages(TArray<FString>& OutRevertablePackageNames)
 {
-	// determine source control locations
-	TArray<FString> SourceControlLocations;
+	if (!ISourceControlModule::Get().IsEnabled() || !ISourceControlModule::Get().GetProvider().IsAvailable())
+	{
+		return false;
+	}
 
+	// update status for all packages
+	TArray<FString> Filenames;
 	if (ISourceControlModule::Get().UsesCustomProjectDir())
 	{
-		SourceControlLocations.Add(ISourceControlModule::Get().GetSourceControlProjectDir());
+		FString SourceControlProjectDir = ISourceControlModule::Get().GetSourceControlProjectDir();
+		Filenames.Add(SourceControlProjectDir);
 	}
 	else
 	{
-		GetSourceControlLocations(/*bContentOnly=*/true);
+		Filenames = GetSourceControlLocations();
 	}
-
-	// find packages in those locations
-	for (const FString& SourceControlLocation : SourceControlLocations)
+	
+	ISourceControlProvider& SourceControlProvider = ISourceControlModule::Get().GetProvider();
+	FSourceControlOperationRef Operation = ISourceControlOperation::Create<FUpdateStatus>();
+	if (SourceControlProvider.Execute(Operation, Filenames) != ECommandResult::Succeeded)
 	{
-		TArray<FString> PackageRelativePaths;
-		if (FPackageName::FindPackagesInDirectory(PackageRelativePaths, FPaths::ConvertRelativePathToFull(SourceControlLocation)))
-		{
-			OutPackageNames.Reserve(OutPackageNames.Num() + PackageRelativePaths.Num());
-			for (const FString& Path : PackageRelativePaths)
-			{
-				FString PackageName;
-				FString FailureReason;
-				if (FPackageName::TryConvertFilenameToLongPackageName(Path, PackageName, &FailureReason))
-				{
-					OutPackageNames.Add(PackageName);
-				}
-				else
-				{
-					FMessageLog("SourceControl").Error(FText::FromString(FailureReason));
-				}
-			}
-		}
+		return false;
 	}
 
-	return true;
-}
-
-// TODO: Move outside of helpers
-bool USourceControlHelpers::RevertAndReloadAllPackages(const TArray<FString>& InFilenames)
-{
-	TArray<FString> PackagesToReload;
-	ListAllPackages(PackagesToReload);
+	// Get a list of all the revertable packages	
+	TMap<FString, FSourceControlStatePtr> PackageStates;
+	FEditorFileUtils::FindAllSubmittablePackageFiles(PackageStates, true);
 
-	if (InFilenames.Num() > 0)
+	for (auto& PackageState : PackageStates)
 	{
-		RevertAndReloadPackages(InFilenames, PackagesToReload, /*bRevertAll=*/false, /*bAllowReloadWorld=*/true);
-		return true;
+		const FString PackageName = PackageState.Key;
+		OutRevertablePackageNames.Add(PackageName);
 	}
 
-	return false;
+	return true;
 }
 
-// TODO: Move outside of helpers
-bool USourceControlHelpers::RevertAllChangesAndReloadAllPackages()
+bool USourceControlHelpers::RevertAllChangesAndReloadWorld()
 {
 	TArray<FString> PackagesToReload;
-	ListAllPackages(PackagesToReload);
+	ListRevertablePackages(PackagesToReload);
 
-	RevertAndReloadPackages(TArray<FString>(), PackagesToReload, /*bRevertAll=*/true, /*bAllowReloadWorld=*/true);
+	RevertAndReloadPackages(PackagesToReload, /*bRevertAll=*/true, /*bReloadWorld=*/true);
 
 	return true;
 }
 
-bool USourceControlHelpers::RevertAndReloadPackages(const TArray<FString>& InPackagesToRevert, const TArray<FString>& InPackagesToReload, bool bRevertAll, bool bAllowReloadWorld)
+bool USourceControlHelpers::RevertAndReloadPackages(const TArray<FString>& InPackagesToRevert, bool bRevertAll, bool bReloadWorld)
 {
 	auto RevertOperation = [bRevertAll](const TArray<FString>& InPackagesToRevert) -> bool
 	{
@@ -1179,7 +1196,7 @@ bool USourceControlHelpers::RevertAndReloadPackages(const TArray<FString>& InPac
 		return SourceControlProvider.Execute(RevertOperation, InPackagesToRevert, EConcurrency::Synchronous, OperationCompleteCallback) == ECommandResult::Succeeded;
 	};
 
-	return ApplyOperationAndReloadPackages(InPackagesToRevert, InPackagesToReload, RevertOperation, bAllowReloadWorld);
+	return ApplyOperationAndReloadPackages(InPackagesToRevert,RevertOperation, bReloadWorld);
 }
 #endif //!WITH_EDITOR
 
diff --git a/Engine/Source/Developer/SourceControl/Public/SourceControlHelpers.h b/Engine/Source/Developer/SourceControl/Public/SourceControlHelpers.h
--- a/Engine/Source/Developer/SourceControl/Public/SourceControlHelpers.h
+++ b/Engine/Source/Developer/SourceControl/Public/SourceControlHelpers.h
@@ -350,39 +350,30 @@ public:
 #if WITH_EDITOR
 
 	/**
-	 * Applies the provided function and reloads provided packages.
+	 * Applies the provided function and optionally reloads the world
 	 * @param	InPackagesToApplyOperation			The files/packages to apply the operation
-	 * @param	InPackagesToReload					The files to reload
 	 * @param	InOperation							The function to apply
-	 * @param	bAllowReloadWorld					Allow reloading of the world
+	 * @param	bReloadWorld						Reload the world
 	 * @return true if succeeded.
 	 */
-	static bool ApplyOperationAndReloadPackages(const TArray<FString>& InPackagesToApplyOperation, const TArray<FString>& InPackagesToReload, 
-		const TFunctionRef<bool(const TArray<FString>&)>& InOperation, bool bAllowReloadWorld = false);
+	static bool ApplyOperationAndReloadPackages(const TArray<FString>& InPackagesToApplyOperation, 
+		const TFunctionRef<bool(const TArray<FString>&)>& InOperation, bool bReloadWorld = false);
 
 	/**
-     * Reverts the provided files then reload packages.
+     * Reverts the provided files then reloads packages.
 	 * @param	InPackagesToRevert					The packages to revert
-	 * @param	InPackagesToReload					The packages to reload
 	 * @param	bRevertAll							Whether to revert all files
-	 * @param	bAllowReloadWorld					Allow reloading of the world
+	 * @param	bReloadWorld						Reload the world
      * @return true if succeeded.
      */
 	UFUNCTION(BlueprintCallable, Category = "Editor Scripting | Editor Revision Control Helpers", meta = (Keywords = "Source Control"))
-	static bool RevertAndReloadPackages(const TArray<FString>& InPackagesToRevert, const TArray<FString>& InPackagesToReload, bool bRevertAll = false, bool bAllowReloadWorld = false);
+	static bool RevertAndReloadPackages(const TArray<FString>& InPackagesToRevert, bool bRevertAll = false, bool bReloadWorld = false);
 
 	/**
-	 * Reverts the provided files and then reloads all packages.
-	 * @param	InFilenames		The files/packages to revert
+	 * Reverts all changes and then reloads the world.
 	 * @return true if succeeded.
 	 */
-	static bool RevertAndReloadAllPackages(const TArray<FString>& InFilenames);
-
-	/**
-	 * Reverts all changes and then reloads all packages.
-	 * @return true if succeeded.
-	 */
-	static bool RevertAllChangesAndReloadAllPackages();
+	static bool RevertAllChangesAndReloadWorld();
 
 #endif // !WITH_EDITOR
 
@@ -632,9 +623,9 @@ public:
 	static bool GetAssetDataFromFileHistory(FSourceControlStatePtr InSourceControlState, TArray<FAssetData>& OutAssets, TArray<FName>* OutDependencies = nullptr, int64 MaxFetchSize = -1);
 
 	/**
-	 * Find packages in Source Control locations
+	 * Find packages that can be reverted in Source Control locations
 	 */
-	static bool ListAllPackages(TArray<FString>& OutPackageNames);
+	static bool ListRevertablePackages(TArray<FString>& OutRevertablePackageNames);
 
 	/**
 	 * Get the list of files and directories that source control should check when looking for changes.
diff --git a/Engine/Source/Developer/SourceControl/Public/SourceControlOperations.h b/Engine/Source/Developer/SourceControl/Public/SourceControlOperations.h
--- a/Engine/Source/Developer/SourceControl/Public/SourceControlOperations.h
+++ b/Engine/Source/Developer/SourceControl/Public/SourceControlOperations.h
@@ -279,6 +279,64 @@ protected:
 	TArray<FString>		DeletedFiles;
 };
 
+/**
+ * Operation used to determine the files that would be affected by a sync operation
+ */
+class FSyncPreview : public FSourceControlOperationBase
+{
+public:
+	// ISourceControlOperation interface
+	virtual FName GetName() const override
+	{
+		return "SyncPreview";
+	}
+
+	virtual FText GetInProgressString() const override
+	{
+		return LOCTEXT("SourceControl_SyncPreview", "Previewing File Sync from revision control...");
+	}
+
+	void SetRevision(const FString& InRevision)
+	{
+		Revision = InRevision;
+	}
+
+	const FString& GetRevision() const
+	{
+		return Revision;
+	}
+
+	void SetHeadRevisionFlag(const bool bInHeadRevision)
+	{
+		bHeadRevision = bInHeadRevision;
+	}
+
+	bool IsHeadRevisionFlagSet() const
+	{
+		return bHeadRevision;
+	}
+
+	void SetAffectedFiles(TArray<FString>&& InAffectedFiles)
+	{
+		AffectedFiles = MoveTemp(InAffectedFiles);
+	}
+
+	const TArray<FString>& GetAffectedFiles() const
+	{
+		return AffectedFiles;
+	}
+	
+protected:
+	/** Target Revision to which the sync preview refers to */
+	FString Revision;
+
+	/** Flag abstracting if the operation aims to preview a sync to head */
+	bool bHeadRevision = false;
+
+	/** Array of files that would be affected by the sync operation */
+	TArray<FString> AffectedFiles;
+};
+
 /**
  * Operation used to sync files to the state they are in source control
  */
diff --git a/Engine/Source/Developer/UncontrolledChangelists/Private/UncontrolledChangelistsModule.cpp b/Engine/Source/Developer/UncontrolledChangelists/Private/UncontrolledChangelistsModule.cpp
--- a/Engine/Source/Developer/UncontrolledChangelists/Private/UncontrolledChangelistsModule.cpp
+++ b/Engine/Source/Developer/UncontrolledChangelists/Private/UncontrolledChangelistsModule.cpp
@@ -366,7 +366,7 @@ bool FUncontrolledChangelistsModule::OnRevert(const TArray<FString>& InFilenames
 		return true;
 	}
 
-	bSuccess = SourceControlHelpers::ApplyOperationAndReloadPackages(InFilenames, InFilenames, ExecuteRevertOperation);
+	bSuccess = SourceControlHelpers::ApplyOperationAndReloadPackages(InFilenames, ExecuteRevertOperation);
 
 	UpdateStatus();
 
diff --git a/Engine/Source/Editor/MainFrame/Private/Frame/MainFrameActions.cpp b/Engine/Source/Editor/MainFrame/Private/Frame/MainFrameActions.cpp
--- a/Engine/Source/Editor/MainFrame/Private/Frame/MainFrameActions.cpp
+++ b/Engine/Source/Editor/MainFrame/Private/Frame/MainFrameActions.cpp
@@ -112,7 +112,7 @@ void FMainFrameCommands::RegisterCommands()
 	ActionList->MapAction( SubmitContent, FExecuteAction::CreateLambda([]() { FSourceControlWindows::ChoosePackagesToCheckIn(); }), FCanExecuteAction::CreateStatic(&FSourceControlWindows::CanChoosePackagesToCheckIn ), FGetActionCheckState(), FIsActionButtonVisible::CreateStatic(FSourceControlWindows::ShouldChoosePackagesToCheckBeVisible) );
 
 	UI_COMMAND( SyncContent, "Sync Content", "Saves all unsaved levels and assets to disk and then downloads the latest versions from revision control.", EUserInterfaceActionType::Button, FInputChord( EKeys::F5 ) );
-	ActionList->MapAction(SyncContent, FExecuteAction::CreateLambda([]() { FSourceControlWindows::SyncAllPackages(); }), FCanExecuteAction::CreateStatic( &FSourceControlWindows::CanSyncAllPackages ) );
+	ActionList->MapAction(SyncContent, FExecuteAction::CreateLambda([]() { FSourceControlWindows::SyncLatest(); }), FCanExecuteAction::CreateStatic( &FSourceControlWindows::CanSyncLatest ) );
 
 	UI_COMMAND( ConnectToSourceControl, "Connect to Revision Control...", "Connect to a revision control system for tracking changes to your content and levels.", EUserInterfaceActionType::Button, FInputChord() );
 	ActionList->MapAction( ConnectToSourceControl, FExecuteAction::CreateStatic( &FMainFrameActionCallbacks::ConnectToSourceControl ), DefaultExecuteAction );
diff --git a/Engine/Source/Editor/SceneOutliner/Private/SceneOutlinerSCCHandler.cpp b/Engine/Source/Editor/SceneOutliner/Private/SceneOutlinerSCCHandler.cpp
--- a/Engine/Source/Editor/SceneOutliner/Private/SceneOutlinerSCCHandler.cpp
+++ b/Engine/Source/Editor/SceneOutliner/Private/SceneOutlinerSCCHandler.cpp
@@ -371,8 +371,7 @@ void FSceneOutlinerSCCHandler::ExecuteSCCRevert()
 {
 	TArray<FString> PackageNames;
 	GetSelectedPackageNames(PackageNames);
-
-	FSourceControlWindows::PromptForRevert(PackageNames, true);
+	FSourceControlWindows::PromptForRevert(PackageNames, /*bReloadWorld=*/true);
 }
 
 void FSceneOutlinerSCCHandler::ExecuteSCCHistory()
diff --git a/Engine/Source/Editor/SourceControlWindows/Private/SSourceControlChangelists.cpp b/Engine/Source/Editor/SourceControlWindows/Private/SSourceControlChangelists.cpp
--- a/Engine/Source/Editor/SourceControlWindows/Private/SSourceControlChangelists.cpp
+++ b/Engine/Source/Editor/SourceControlWindows/Private/SSourceControlChangelists.cpp
@@ -1967,7 +1967,7 @@ void SSourceControlChangelistsWidget::OnRevert()
 
 			SSourceControlCommon::ExecuteChangelistOperationWithSlowTaskWrapper(LOCTEXT("Reverting_Files", "Reverting file(s)..."), [&PackageFiles, &NonPackageFiles]()
 			{
-				if ((PackageFiles.Num() == 0 || SourceControlHelpers::RevertAndReloadPackages(PackageFiles, PackageFiles)) && (NonPackageFiles.Num() == 0 || SourceControlHelpers::RevertFiles(NonPackageFiles)))
+				if ((PackageFiles.Num() == 0 || SourceControlHelpers::RevertAndReloadPackages(PackageFiles)) && (NonPackageFiles.Num() == 0 || SourceControlHelpers::RevertFiles(NonPackageFiles)))
 				{
 					SSourceControlCommon::DisplaySourceControlOperationNotification(LOCTEXT("Revert_Files_Succeeded", "The selected file(s) were reverted."), SNotificationItem::CS_Success);
 				}
diff --git a/Engine/Source/Editor/SourceControlWindows/Private/SSourceControlRevert.cpp b/Engine/Source/Editor/SourceControlWindows/Private/SSourceControlRevert.cpp
--- a/Engine/Source/Editor/SourceControlWindows/Private/SSourceControlRevert.cpp
+++ b/Engine/Source/Editor/SourceControlWindows/Private/SSourceControlRevert.cpp
@@ -397,7 +397,7 @@ private:
 	bool bRevertUnchangedFilesOnly;
 };
 
-bool FSourceControlWindows::PromptForRevert( const TArray<FString>& InPackageNames, bool InReloadAllPackages)
+bool FSourceControlWindows::PromptForRevert( const TArray<FString>& InPackageNames, bool bInReloadWorld)
 {
 	bool bReverted = false;
 
@@ -440,14 +440,7 @@ bool FSourceControlWindows::PromptForRevert( const TArray<FString>& InPackageNam
 			
 			if (FinalPackagesToRevert.Num() > 0)
 			{
-				if (InReloadAllPackages)
-				{
-					SourceControlHelpers::RevertAndReloadAllPackages(FinalPackagesToRevert);
-				}
-				else
-				{
-					SourceControlHelpers::RevertAndReloadPackages(FinalPackagesToRevert, FinalPackagesToRevert, /*bRevertAll=*/false, /*bAllowReloadWorld=*/true);
-				}
+				SourceControlHelpers::RevertAndReloadPackages(FinalPackagesToRevert, /*bRevertAll=*/false, /*bReloadWorld=*/bInReloadWorld);
 
 				bReverted = true;
 			}
@@ -457,9 +450,9 @@ bool FSourceControlWindows::PromptForRevert( const TArray<FString>& InPackageNam
 	return bReverted;
 }
 
-bool FSourceControlWindows::RevertAllChangesAndReloadAllPackages()
+bool FSourceControlWindows::RevertAllChangesAndReloadWorld()
 {	
-	return SourceControlHelpers::RevertAllChangesAndReloadAllPackages();
+	return SourceControlHelpers::RevertAllChangesAndReloadWorld();
 }
 
 #undef LOCTEXT_NAMESPACE
diff --git a/Engine/Source/Editor/SourceControlWindows/Private/SourceControlWindows.cpp b/Engine/Source/Editor/SourceControlWindows/Private/SourceControlWindows.cpp
--- a/Engine/Source/Editor/SourceControlWindows/Private/SourceControlWindows.cpp
+++ b/Engine/Source/Editor/SourceControlWindows/Private/SourceControlWindows.cpp
@@ -178,31 +178,21 @@ static bool SaveDirtyPackages()
 	return bSaved;
 }
 
-bool FSourceControlWindows::SyncAllPackages()
+bool FSourceControlWindows::SyncLatest()
 {
 	bool bSaved = SaveDirtyPackages();
 
 	// if properly saved or confirmation given, find all packages and use source control to update them.
 	if (bSaved)
 	{
-		TArray<FString> PackageNames;
-		if (SourceControlHelpers::ListAllPackages(PackageNames))
-		{
-			// sync those packages if any were found
-			if (PackageNames.Num() > 0)
-			{
-				AssetViewUtils::SyncPackagesFromSourceControl(PackageNames);
-			}
-
-			return true;
-		}
+		AssetViewUtils::SyncLatestFromSourceControl();
 	}
 
 	return false;
 }
 
 
-bool FSourceControlWindows::CanSyncAllPackages()
+bool FSourceControlWindows::CanSyncLatest()
 {
 	ISourceControlModule& SourceControlModule = ISourceControlModule::Get();
 
diff --git a/Engine/Source/Editor/SourceControlWindows/Public/SourceControlWindows.h b/Engine/Source/Editor/SourceControlWindows/Public/SourceControlWindows.h
--- a/Engine/Source/Editor/SourceControlWindows/Public/SourceControlWindows.h
+++ b/Engine/Source/Editor/SourceControlWindows/Public/SourceControlWindows.h
@@ -48,16 +48,16 @@ public:
 
 	/** Determines if the Submit Content action should be visible or not */
 	static bool ShouldChoosePackagesToCheckBeVisible();
+	
 	/**
-	 * Saves all unsaved levels and assets and then - conditionally - performs an FSync operation on all packages, unloading and reloading
-	 * them as necessary.
+	 * Saves all unsaved levels and assets and then - conditionally - performs an FSync operation to get latest and reloads the world
 	 * 
 	 * @return	true - if command completed successfully.
 	 */
-	static bool SyncAllPackages();
+	static bool SyncLatest();
 
-	/** Determines whether we can sync all packages */
-	static bool CanSyncAllPackages();
+	/** Determines whether we can sync to latest */
+	static bool CanSyncLatest();
 
 	/**
 	 * Display check in dialog for the specified packages and get additional result information
@@ -95,16 +95,16 @@ public:
 	 * Prompt the user with a revert files dialog, allowing them to specify which packages, if any, should be reverted.
 	 *
 	 * @param	InPackageNames	Names of the packages to consider for reverting
-	 * @param	InReloadAllPackages	Reload all packages as part of the revert operation
+	 * @param	bInReloadWorld	Reload the world as part of the revert operation
 	 *
 	 * @return	true if the files were reverted; false if the user canceled out of the dialog
 	 */
-	static bool PromptForRevert(const TArray<FString>& InPackageNames, bool InReloadAllPackages = false );
+	static bool PromptForRevert(const TArray<FString>& InPackageNames, bool bInReloadWorld = false );
 
 	/**
-	 * Revert all locally modified files, and reload all packages
+	 * Revert all locally modified files, and reload the world
 	 */
-	static bool RevertAllChangesAndReloadAllPackages();
+	static bool RevertAllChangesAndReloadWorld();
 
 	/**
 	 * Displays file diff against workspace version
diff --git a/Engine/Source/Editor/StatusBar/Private/SourceControlMenuHelpers.cpp b/Engine/Source/Editor/StatusBar/Private/SourceControlMenuHelpers.cpp
--- a/Engine/Source/Editor/StatusBar/Private/SourceControlMenuHelpers.cpp
+++ b/Engine/Source/Editor/StatusBar/Private/SourceControlMenuHelpers.cpp
@@ -160,7 +160,7 @@ void FSourceControlCommands::RevertAllModifiedFiles_Clicked()
 	FText Title = LOCTEXT("RevertAllModifiedFiles_Title", "Revert all local changes");
 	if (FMessageDialog::Open(EAppMsgType::YesNo, EAppReturnType::No, Message, &Title) == EAppReturnType::Yes)
 	{
-		FSourceControlWindows::RevertAllChangesAndReloadAllPackages();
+		FSourceControlWindows::RevertAllChangesAndReloadWorld();
 	}
 }
 
@@ -388,9 +388,9 @@ const FSlateBrush* FSourceControlMenuHelpers::GetSourceControlSyncStatusIcon()
 
 FReply FSourceControlMenuHelpers::OnSourceControlSyncClicked()
 {
-	if (FSourceControlWindows::CanSyncAllPackages())
+	if (FSourceControlWindows::CanSyncLatest())
 	{
-		FSourceControlWindows::SyncAllPackages();
+		FSourceControlWindows::SyncLatest();
 	}
 
 	return FReply::Handled();
