Ensure the camera does not reset when we are syncing or reverting a file or all files via source control
A scoped class is added to retain and restore all camera views which is used when we reload or sync packages.

diff --git a/Engine/Source/Developer/AssetTools/Private/AssetViewUtils.cpp b/Engine/Source/Developer/AssetTools/Private/AssetViewUtils.cpp
--- a/Engine/Source/Developer/AssetTools/Private/AssetViewUtils.cpp
+++ b/Engine/Source/Developer/AssetTools/Private/AssetViewUtils.cpp
@@ -41,7 +41,6 @@
 #include "Subsystems/AssetEditorSubsystem.h"
 #include "Editor.h"
 #include "UObject/LinkerInstancingContext.h"
-#include "LevelEditorViewport.h"
 
 #define LOCTEXT_NAMESPACE "ContentBrowser"
 
@@ -1776,18 +1775,6 @@ void AssetViewUtils::SyncPackagesFromSourceControl(const TArray<FString>& Packag
 			return false; // keep package
 		});
 
-		// Syncing may result in a map reload, so we retain the camera views to restore after the reload...
-		struct FCameraView
-		{
-			FVector Location;
-			FRotator Rotation;
-		};
-		TMap<FLevelEditorViewportClient*, FCameraView> CameraViews;
-		for (FLevelEditorViewportClient* LevelVC : GEditor->GetLevelViewportClients())
-		{
-			CameraViews.Add(LevelVC, { LevelVC->GetViewLocation(), LevelVC->GetViewRotation() });
-		}
-
 		// Hot-reload the new packages...
 		if (bIsSyncLatestOperation)
 		{
@@ -1796,19 +1783,6 @@ void AssetViewUtils::SyncPackagesFromSourceControl(const TArray<FString>& Packag
 		}
 		UPackageTools::ReloadPackages(LoadedPackages);
 
-		// Restore the camera views after a map reload if necessary...
-		for (FLevelEditorViewportClient* LevelVC : GEditor->GetLevelViewportClients())
-		{
-			const FCameraView& CameraView = CameraViews.FindChecked(LevelVC);
-			LevelVC->SetViewLocation(CameraView.Location);
-			if (!LevelVC->IsOrtho())
-			{
-				LevelVC->SetViewRotation(CameraView.Rotation);
-			}
-			LevelVC->Invalidate();
-			FEditorDelegates::OnEditorCameraMoved.Broadcast(CameraView.Location, CameraView.Rotation, LevelVC->ViewportType, LevelVC->ViewIndex);
-		}
-
 		// Unload any deleted packages...
 		TArray<UPackage*> PackageRawPtrsToUnload;
 		for (TWeakObjectPtr<UPackage>& PackageToUnload : PackagesToUnload)
diff --git a/Engine/Source/Developer/SourceControl/Private/SourceControlHelpers.cpp b/Engine/Source/Developer/SourceControl/Private/SourceControlHelpers.cpp
--- a/Engine/Source/Developer/SourceControl/Private/SourceControlHelpers.cpp
+++ b/Engine/Source/Developer/SourceControl/Private/SourceControlHelpers.cpp
@@ -25,7 +25,6 @@
 #include "PackageTools.h"
 #include "ObjectTools.h"
 #include "FileHelpers.h"
-#include "LevelEditorViewport.h"
 #endif
 
 #define LOCTEXT_NAMESPACE "SourceControlHelpers"
@@ -1036,22 +1035,6 @@ bool USourceControlHelpers::ApplyOperationAndReloadPackages(const TArray<FString
 		return false; // keep package
 	});
 
-	struct FCameraView
-	{
-		FVector Location;
-		FRotator Rotation;
-	};
-	TMap<FLevelEditorViewportClient*, FCameraView> CameraViews;
-
-	if (bReloadWorld)
-	{
-		// As we're reloading the world, we retain the camera views so they can be restored after the reload.
-		for (FLevelEditorViewportClient* LevelVC : GEditor->GetLevelViewportClients())
-		{
-			CameraViews.Add(LevelVC, { LevelVC->GetViewLocation(), LevelVC->GetViewRotation() });
-		}
-	}
-
 	// Hot-reload the new packages...
 	FText OutReloadErrorMsg;
 	const bool bInteractive = true;
@@ -1061,22 +1044,6 @@ bool USourceControlHelpers::ApplyOperationAndReloadPackages(const TArray<FString
 		UE_LOG(LogSourceControl, Warning, TEXT("%s"), *OutReloadErrorMsg.ToString());
 	}
 
-	if (bReloadWorld)
-	{
-		// Restore the camera views after a world reload if necessary...
-		for (FLevelEditorViewportClient* LevelVC : GEditor->GetLevelViewportClients())
-		{
-			const FCameraView& CameraView = CameraViews.FindChecked(LevelVC);
-			LevelVC->SetViewLocation(CameraView.Location);
-			if (!LevelVC->IsOrtho())
-			{
-				LevelVC->SetViewRotation(CameraView.Rotation);
-			}
-			LevelVC->Invalidate();
-			FEditorDelegates::OnEditorCameraMoved.Broadcast(CameraView.Location, CameraView.Rotation, LevelVC->ViewportType, LevelVC->ViewIndex);
-		}
-	}
-
 	// Delete and Unload assets...
 	if(ObjectTools::DeleteObjectsUnchecked(ObjectsToDelete) != ObjectsToDelete.Num())
 	{ 
diff --git a/Engine/Source/Editor/SceneOutliner/Private/SceneOutlinerSCCHandler.cpp b/Engine/Source/Editor/SceneOutliner/Private/SceneOutlinerSCCHandler.cpp
--- a/Engine/Source/Editor/SceneOutliner/Private/SceneOutlinerSCCHandler.cpp
+++ b/Engine/Source/Editor/SceneOutliner/Private/SceneOutlinerSCCHandler.cpp
@@ -10,6 +10,7 @@
 #include "UncontrolledChangelistsModule.h"
 #include "AssetViewUtils.h"
 #include "RevisionControlStyle/RevisionControlStyle.h"
+#include "Bookmarks/BookmarkScoped.h"
 
 #define LOCTEXT_NAMESPACE "FSceneOutlinerSCCHandler"
 
@@ -395,6 +396,7 @@ void FSceneOutlinerSCCHandler::ExecuteSCCRevert()
 
 	if (PackagesToRevert.Num() > 0)
 	{
+		FBookmarkScoped BookmarkScoped;
 		SourceControlHelpers::RevertAndReloadPackages(PackagesToRevert, /*bRevertAll=*/false, /*bReloadWorld=*/true);
 	}
 }
diff --git a/Engine/Source/Editor/StatusBar/Private/SourceControlMenuHelpers.cpp b/Engine/Source/Editor/StatusBar/Private/SourceControlMenuHelpers.cpp
--- a/Engine/Source/Editor/StatusBar/Private/SourceControlMenuHelpers.cpp
+++ b/Engine/Source/Editor/StatusBar/Private/SourceControlMenuHelpers.cpp
@@ -24,6 +24,7 @@
 #include "Misc/ConfigCacheIni.h"
 #include "Misc/MessageDialog.h"
 #include "RevisionControlStyle/RevisionControlStyle.h"
+#include "Bookmarks/BookmarkScoped.h"
 
 #define LOCTEXT_NAMESPACE "SourceControlCommands"
 
@@ -160,6 +161,7 @@ void FSourceControlCommands::RevertAllModifiedFiles_Clicked()
 	FText Title = LOCTEXT("RevertAllModifiedFiles_Title", "Revert all local changes");
 	if (FMessageDialog::Open(EAppMsgType::YesNo, EAppReturnType::No, Message, &Title) == EAppReturnType::Yes)
 	{
+		FBookmarkScoped BookmarkScoped;
 		FSourceControlWindows::RevertAllChangesAndReloadWorld();
 	}
 }
@@ -390,6 +392,7 @@ FReply FSourceControlMenuHelpers::OnSourceControlSyncClicked()
 {
 	if (FSourceControlWindows::CanSyncLatest())
 	{
+		FBookmarkScoped BookmarkScoped;
 		FSourceControlWindows::SyncLatest();
 	}
 
diff --git a/Engine/Source/Editor/UnrealEd/Private/Bookmarks/BookmarkScoped.cpp b/Engine/Source/Editor/UnrealEd/Private/Bookmarks/BookmarkScoped.cpp
new file mode 100644
--- /dev/null
+++ b/Engine/Source/Editor/UnrealEd/Private/Bookmarks/BookmarkScoped.cpp
@@ -0,0 +1,35 @@
+// Copyright Epic Games, Inc. All Rights Reserved.
+
+#include "Bookmarks/BookmarkScoped.h"
+
+#include "Bookmarks/BookmarkSingleViewportActions.h"
+
+FBookmarkScoped::FBookmarkScoped() :
+	Actions(MakeShared<FBookmarkSingleViewportActions>())
+{
+	// Bookmark all viewports to restore when going out of scope
+	for (FEditorViewportClient* ViewportClient : GEditor->GetAllViewportClients())
+	{
+		UBookMark* BookMark = NewObject<UBookMark>();
+		BookMark->AddToRoot();
+		Actions->InitFromViewport(BookMark, *ViewportClient);
+		BookMarks.Add(ViewportClient, BookMark);
+	}
+}
+
+
+FBookmarkScoped::~FBookmarkScoped()
+{
+	// Restore all viewport bookmarks
+	for (FEditorViewportClient* ViewportClient : GEditor->GetAllViewportClients())
+	{
+		UBookMark* BookMark = BookMarks.FindChecked(ViewportClient);
+		TSharedPtr<FBookmarkBaseJumpToSettings> Settings = MakeShared<FBookmarkJumpToSettings>();
+		Actions->JumpToBookmark(BookMark, Settings, *ViewportClient);
+	}
+
+	for (TPair<FEditorViewportClient*, UBookMark*>& Pair : BookMarks)
+	{
+		Pair.Value->RemoveFromRoot();
+	}
+}
diff --git a/Engine/Source/Editor/UnrealEd/Public/Bookmarks/BookMarkTypeActions.h b/Engine/Source/Editor/UnrealEd/Public/Bookmarks/BookMarkTypeActions.h
--- a/Engine/Source/Editor/UnrealEd/Public/Bookmarks/BookMarkTypeActions.h
+++ b/Engine/Source/Editor/UnrealEd/Public/Bookmarks/BookMarkTypeActions.h
@@ -56,15 +56,22 @@ public:
 			// Set all level editing cameras to this bookmark
 			for (FLevelEditorViewportClient* LevelVC : GEditor->GetLevelViewportClients())
 			{
-				LevelVC->SetViewLocation(Bookmark->Location);
-				if (!LevelVC->IsOrtho())
-				{
-					LevelVC->SetViewRotation(Bookmark->Rotation);
-				}
-				LevelVC->Invalidate();
-
-				FEditorDelegates::OnEditorCameraMoved.Broadcast(Bookmark->Location, Bookmark->Rotation, LevelVC->ViewportType, LevelVC->ViewIndex);
+				ApplyBookmarkToViewportClient(Bookmark, LevelVC);
 			}
 		}
 	}
+
+protected:
+
+	void ApplyBookmarkToViewportClient(UBookMark* InBookmark, FEditorViewportClient* InViewportClient)
+	{
+		InViewportClient->SetViewLocation(InBookmark->Location);
+		if (!InViewportClient->IsOrtho())
+		{
+			InViewportClient->SetViewRotation(InBookmark->Rotation);
+		}
+		InViewportClient->Invalidate();
+
+		FEditorDelegates::OnEditorCameraMoved.Broadcast(InBookmark->Location, InBookmark->Rotation, InViewportClient->ViewportType, InViewportClient->ViewIndex);
+	}
 };
diff --git a/Engine/Source/Editor/UnrealEd/Public/Bookmarks/BookmarkScoped.h b/Engine/Source/Editor/UnrealEd/Public/Bookmarks/BookmarkScoped.h
new file mode 100644
--- /dev/null
+++ b/Engine/Source/Editor/UnrealEd/Public/Bookmarks/BookmarkScoped.h
@@ -0,0 +1,29 @@
+// Copyright Epic Games, Inc. All Rights Reserved.
+
+#pragma once
+
+class IBookmarkTypeActions;
+class UBookMark;
+
+/**
+ * Provides a way to temporarily bookmark and restore all viewports
+ * This allows camera locations and rotations to be preserved across a map reload
+ */
+class UNREALED_API FBookmarkScoped
+{
+public:
+
+	/** Constructor; bookmark all viewports */
+	FBookmarkScoped();
+
+	/** Destructor; restore all viewports */
+	~FBookmarkScoped();
+
+private:
+
+	/** Action which allows a single viewport to jump to a single bookmark */
+	TSharedPtr<IBookmarkTypeActions> Actions;
+
+	/** Bookmarks for each viewport */
+	TMap<FEditorViewportClient*, UBookMark*> BookMarks;
+};
diff --git a/Engine/Source/Editor/UnrealEd/Public/Bookmarks/BookmarkSingleViewportActions.h b/Engine/Source/Editor/UnrealEd/Public/Bookmarks/BookmarkSingleViewportActions.h
new file mode 100644
--- /dev/null
+++ b/Engine/Source/Editor/UnrealEd/Public/Bookmarks/BookmarkSingleViewportActions.h
@@ -0,0 +1,21 @@
+// Copyright Epic Games, Inc. All Rights Reserved.
+
+#pragma once
+
+#include "Bookmarks/BookMarkTypeActions.h"
+
+/**
+ * A bookmark action which jumps a single viewport to a single bookmark
+ */
+class FBookmarkSingleViewportActions : public FBookMarkTypeActions
+{
+public:
+
+	virtual void JumpToBookmark(UBookmarkBase* InBookmark, const TSharedPtr<struct FBookmarkBaseJumpToSettings> InSettings, FEditorViewportClient& InViewportClient) override
+	{
+		if (UBookMark* Bookmark = Cast<UBookMark>(InBookmark))
+		{
+			ApplyBookmarkToViewportClient(Bookmark, &InViewportClient);
+		}
+	}
+};
