SourceControl: Refactored UsesCustomProjectDir / GetSourceControlLocations usage pattern.

diff --git a/Engine/Source/Developer/AssetTools/Private/AssetViewUtils.cpp b/Engine/Source/Developer/AssetTools/Private/AssetViewUtils.cpp
--- a/Engine/Source/Developer/AssetTools/Private/AssetViewUtils.cpp
+++ b/Engine/Source/Developer/AssetTools/Private/AssetViewUtils.cpp
@@ -1828,16 +1828,9 @@ void AssetViewUtils::SyncLatestFromSourceControl()
 	else
 	{
 		// Likely the source control provider does not support sync previews. Proceed to syncing all paths under source control
-		TArray<FString> SourceControlLocations;
 
-		if (ISourceControlModule::Get().UsesCustomProjectDir())
-		{
-			SourceControlLocations.Add(ISourceControlModule::Get().GetSourceControlProjectDir());
-		}
-		else
-		{
-			SourceControlLocations = SourceControlHelpers::GetSourceControlLocations(/*bContentOnly=*/true);
-		}
+		bool bContentOnly = !ISourceControlModule::Get().UsesCustomProjectDir();
+		TArray<FString> SourceControlLocations = SourceControlHelpers::GetSourceControlLocations(bContentOnly);
 
 		SyncPathsFromSourceControl(SourceControlLocations);
 	}
diff --git a/Engine/Source/Developer/SourceControl/Private/SourceControlHelpers.cpp b/Engine/Source/Developer/SourceControl/Private/SourceControlHelpers.cpp
--- a/Engine/Source/Developer/SourceControl/Private/SourceControlHelpers.cpp
+++ b/Engine/Source/Developer/SourceControl/Private/SourceControlHelpers.cpp
@@ -1069,22 +1069,45 @@ TArray<FString> USourceControlHelpers::GetSourceControlLocations(const bool bCon
 {
 	TArray<FString> SourceControlLocations;
 
+	if (ISourceControlModule::Get().UsesCustomProjectDir())
 	{
+		FString ProjectDir = ISourceControlModule::Get().GetSourceControlProjectDir();
+
 		TArray<FString> RootPaths;
 		FPackageName::QueryRootContentPaths(RootPaths);
 		for (const FString& RootPath : RootPaths)
 		{
 			const FString RootPathOnDisk = FPackageName::LongPackageNameToFilename(RootPath);
-			SourceControlLocations.Add(FPaths::ConvertRelativePathToFull(RootPathOnDisk));
+			if (FPaths::IsUnderDirectory(RootPathOnDisk, ProjectDir))
+			{
+				SourceControlLocations.Add(FPaths::ConvertRelativePathToFull(RootPathOnDisk));
+			}
 		}
-	}
 
-	if (!bContentOnly)
+		if (!bContentOnly)
+		{
+			SourceControlLocations.Add(ProjectDir);
+		}
+	}
+	else
 	{
-		SourceControlLocations.Add(FPaths::ConvertRelativePathToFull(FPaths::ProjectConfigDir()));
-		SourceControlLocations.Add(FPaths::ConvertRelativePathToFull(FPaths::GetProjectFilePath()));
+		TArray<FString> RootPaths;
+		FPackageName::QueryRootContentPaths(RootPaths);
+		for (const FString& RootPath : RootPaths)
+		{
+			const FString RootPathOnDisk = FPackageName::LongPackageNameToFilename(RootPath);
+			SourceControlLocations.Add(FPaths::ConvertRelativePathToFull(RootPathOnDisk));
+		}
+		
+		if (!bContentOnly)
+		{
+			SourceControlLocations.Add(FPaths::ConvertRelativePathToFull(FPaths::ProjectConfigDir()));
+			SourceControlLocations.Add(FPaths::ConvertRelativePathToFull(FPaths::GetProjectFilePath()));
+		}
 	}
 
+
+
 	return SourceControlLocations;
 }
 
@@ -1096,16 +1119,7 @@ bool USourceControlHelpers::ListRevertablePackages(TArray<FString>& OutRevertabl
 	}
 
 	// update status for all packages
-	TArray<FString> Filenames;
-	if (ISourceControlModule::Get().UsesCustomProjectDir())
-	{
-		FString SourceControlProjectDir = ISourceControlModule::Get().GetSourceControlProjectDir();
-		Filenames.Add(SourceControlProjectDir);
-	}
-	else
-	{
-		Filenames = GetSourceControlLocations();
-	}
+	TArray<FString> Filenames = GetSourceControlLocations();
 	
 	ISourceControlProvider& SourceControlProvider = ISourceControlModule::Get().GetProvider();
 	FSourceControlOperationRef Operation = ISourceControlOperation::Create<FUpdateStatus>();
diff --git a/Engine/Source/Editor/ContentBrowser/Private/FrontendFilters.cpp b/Engine/Source/Editor/ContentBrowser/Private/FrontendFilters.cpp
--- a/Engine/Source/Editor/ContentBrowser/Private/FrontendFilters.cpp
+++ b/Engine/Source/Editor/ContentBrowser/Private/FrontendFilters.cpp
@@ -849,16 +849,8 @@ void FFrontendFilter_NotSourceControlled::RequestStatus()
 		// Request the state of files at filter construction time to make sure files have the correct state for the filter
 		TSharedRef<FUpdateStatus, ESPMode::ThreadSafe> UpdateStatusOperation = ISourceControlOperation::Create<FUpdateStatus>();
 
-		TArray<FString> Filenames;
-		if (ISourceControlModule::Get().UsesCustomProjectDir())
-		{
-			FString SourceControlProjectDir = ISourceControlModule::Get().GetSourceControlProjectDir();
-			Filenames.Add(SourceControlProjectDir);
-		}
-		else
-		{
-			Filenames = SourceControlHelpers::GetSourceControlLocations(/*bContentOnly*/true);
-		}
+		TArray<FString> Filenames = SourceControlHelpers::GetSourceControlLocations(/*bContentOnly*/true);
+
 		UpdateStatusOperation->SetCheckingAllFiles(false);
 		SourceControlProvider.Execute(UpdateStatusOperation, Filenames, EConcurrency::Asynchronous, FSourceControlOperationComplete::CreateSP(this, &FFrontendFilter_NotSourceControlled::SourceControlOperationComplete));
 	}
diff --git a/Engine/Source/Editor/SourceControlWindows/Private/SourceControlWindows.cpp b/Engine/Source/Editor/SourceControlWindows/Private/SourceControlWindows.cpp
--- a/Engine/Source/Editor/SourceControlWindows/Private/SourceControlWindows.cpp
+++ b/Engine/Source/Editor/SourceControlWindows/Private/SourceControlWindows.cpp
@@ -65,16 +65,7 @@ bool FSourceControlWindows::ChoosePackagesToCheckIn(const FSourceControlWindowsO
 	// Start selection process...
 
 	// make sure we update the SCC status of all packages (this could take a long time, so we will run it as a background task)
-	TArray<FString> Filenames;
-	if (ISourceControlModule::Get().UsesCustomProjectDir())
-	{
-		FString SourceControlProjectDir = ISourceControlModule::Get().GetSourceControlProjectDir();
-		Filenames.Add(SourceControlProjectDir);
-	}
-	else
-	{
-		Filenames = SourceControlHelpers::GetSourceControlLocations();
-	}
+	TArray<FString> Filenames = SourceControlHelpers::GetSourceControlLocations();
 	
 	// make sure the SourceControlProvider state cache is populated as well
 	ISourceControlProvider& SourceControlProvider = ISourceControlModule::Get().GetProvider();
@@ -500,19 +491,8 @@ void FSourceControlWindows::ChoosePackagesToCheckInCompleted(const TArray<UPacka
 		return;
 	}
 
-	TArray<FString> PendingDeletePaths;
-
 	bool bUseSourceControlStateCache = true;
-	if (ISourceControlModule::Get().UsesCustomProjectDir())
-	{
-		FString SourceControlProjectDir = ISourceControlModule::Get().GetSourceControlProjectDir();
-		PendingDeletePaths.Add(SourceControlProjectDir);
-		bUseSourceControlStateCache = false;
-	}
-	else
-	{
-		PendingDeletePaths = SourceControlHelpers::GetSourceControlLocations();
-	}
+	TArray<FString> PendingDeletePaths = SourceControlHelpers::GetSourceControlLocations();
 
 	PromptForCheckin(OutResultInfo, PackageNames, PendingDeletePaths, ConfigFiles, bUseSourceControlStateCache);
 }
