SourceControl: Changing capitalization of 'revision control' to 'Revision Control' per UX guidelines.

diff --git a/Engine/Source/Developer/SourceControl/Public/SourceControlOperations.h b/Engine/Source/Developer/SourceControl/Public/SourceControlOperations.h
--- a/Engine/Source/Developer/SourceControl/Public/SourceControlOperations.h
+++ b/Engine/Source/Developer/SourceControl/Public/SourceControlOperations.h
@@ -36,7 +36,7 @@ public:
 
 	virtual FText GetInProgressString() const override
 	{
-		return LOCTEXT("SourceControl_Connecting", "Connecting to revision control...");
+		return LOCTEXT("SourceControl_Connecting", "Connecting to Revision Control...");
 	}
 
 	const FString& GetPassword() const
@@ -293,7 +293,7 @@ public:
 
 	virtual FText GetInProgressString() const override
 	{
-		return LOCTEXT("SourceControl_SyncPreview", "Previewing File Sync from revision control...");
+		return LOCTEXT("SourceControl_SyncPreview", "Previewing sync file(s) from Revision Control...");
 	}
 
 	void SetRevision(const FString& InRevision)
@@ -351,7 +351,7 @@ public:
 
 	virtual FText GetInProgressString() const override
 	{
-		return LOCTEXT("SourceControl_Sync", "Syncing file(s) from revision control...");
+		return LOCTEXT("SourceControl_Sync", "Syncing file(s) from Revision Control...");
 	}
 
 	UE_DEPRECATED(4.26, "FSync::SetRevisionNumber(int32) has been deprecated. Please update to Fsync::SetRevision(const FString&).")
@@ -439,7 +439,7 @@ public:
 
 	virtual FText GetInProgressString() const override
 	{
-		return LOCTEXT("SourceControl_Update", "Updating file(s) revision control status...");
+		return LOCTEXT("SourceControl_Update", "Updating file(s) Revision Control status...");
 	}
 
 	void SetUpdateHistory( bool bInUpdateHistory )
@@ -1245,7 +1245,7 @@ public:
 
 	virtual FText GetInProgressString() const override
 	{
-		return LOCTEXT("SourceControl_GetFile", "Retrieving file from revision control...");
+		return LOCTEXT("SourceControl_GetFile", "Retrieving file from Revision Control...");
 	}
 
 	const FString& GetChangelistNumber() const { return ChangelistNumber; }
