SkeinUE: Rename 'Skein' to 'Unreal Revision Control' in the revision control dropdown.

Tricky bit is to ensure that we're backwards compatible for projects with SourceControlSettings saved on disk/cloud that reference 'Skein' now use 'Unreal Revision Control' instead.

diff --git a/Engine/Source/Developer/SourceControl/Private/SSourceControlPicker.cpp b/Engine/Source/Developer/SourceControl/Private/SSourceControlPicker.cpp
--- a/Engine/Source/Developer/SourceControl/Private/SSourceControlPicker.cpp
+++ b/Engine/Source/Developer/SourceControl/Private/SSourceControlPicker.cpp
@@ -121,10 +121,6 @@ FText SSourceControlPicker::GetProviderText(const FName& InName) const
 	{
 		return LOCTEXT( "GitBetaProviderName", "Git  (beta version)" );
 	}
-	if(InName == "Skein")
-	{
-		return LOCTEXT("SkeinBetaProviderName", "Skein  (beta version)");
-	}
 
 	return FText::FromName(InName);
 }
