SourceControl: Fixed inconsistent results when reverting asset in 'Submit Files' dialog.

diff --git a/Engine/Source/Developer/SourceControl/Private/SourceControlHelpers.cpp b/Engine/Source/Developer/SourceControl/Private/SourceControlHelpers.cpp
--- a/Engine/Source/Developer/SourceControl/Private/SourceControlHelpers.cpp
+++ b/Engine/Source/Developer/SourceControl/Private/SourceControlHelpers.cpp
@@ -14,6 +14,7 @@
 #include "ISourceControlLabel.h"
 #include "UObject/Linker.h"
 #include "UObject/Package.h"
+#include "UObject/UObjectIterator.h"
 #include "Misc/PackageName.h"
 #include "Logging/MessageLog.h"
 #include "AssetRegistry/AssetRegistryModule.h"
@@ -911,10 +912,11 @@ bool USourceControlHelpers::RevertFile(const FString& InFile, bool bSilent)
 bool USourceControlHelpers::ApplyOperationAndReloadPackages(const TArray<FString>& InFilenames,	const TFunctionRef<bool(const TArray<FString>&)>& InOperation, bool bReloadWorld, bool bInteractive)
 {
 	ISourceControlProvider& SourceControlProvider = ISourceControlModule::Get().GetProvider();
-	TArray<UPackage*> LoadedPackages;
+
 	TArray<FString> PackageNames;
 	TArray<FString> PackageFilenames;
 	TArray<FString> FilteredActorPackages;
+	TArray<FString> NotFoundPackages;
 	bool bSuccess = false;
 
 	auto DetachLinker = [](UPackage* Package)
@@ -942,9 +944,10 @@ bool USourceControlHelpers::ApplyOperationAndReloadPackages(const TArray<FString
 		}
 	}
 
-	// Remove packages if they are loaded actors or world
-	bool bWorldFound = false;
-	PackageNames.RemoveAll([&FilteredActorPackages, &LoadedPackages, bReloadWorld, &DetachLinker, &bWorldFound](const FString& PackageName) -> bool
+	// If bReloadWorld=false, remove packages if they are loaded actors or world
+	// If bReloadWorld=true, include world packages to reload
+	TSet<UPackage*> UniqueLoadedPackages;
+	PackageNames.RemoveAll([&FilteredActorPackages, &UniqueLoadedPackages, &NotFoundPackages, bReloadWorld, &DetachLinker](const FString& PackageName) -> bool
 	{
 		UPackage* Package = FindPackage(NULL, *PackageName);
 		
@@ -957,21 +960,18 @@ bool USourceControlHelpers::ApplyOperationAndReloadPackages(const TArray<FString
 					FilteredActorPackages.Emplace(PackageName);
 					return true; // remove the package
 				}
-				else
-				{
-					bWorldFound = true;
-				}
 			}
 			else if (AActor* Actor = AActor::FindActorInPackage(Package))
 			{
 				if (bReloadWorld)
 				{
+					// detach linker on the actor
 					DetachLinker(Package);
 
-					if (!bWorldFound && Actor->GetWorld())
+					// but track its world for reloading - not the actor package itself
+					if (Actor->GetWorld() && Actor->GetWorld()->GetPackage())
 					{
-						bWorldFound = true;
-						LoadedPackages.Add(Actor->GetWorld()->GetPackage());
+						UniqueLoadedPackages.Add(Actor->GetWorld()->GetPackage());
 					}
 
 					return false;
@@ -980,10 +980,14 @@ bool USourceControlHelpers::ApplyOperationAndReloadPackages(const TArray<FString
 				{
 					FilteredActorPackages.Emplace(PackageName);
 					return true; // remove the package
-				}				
+				}
 			}
 
-			LoadedPackages.Add(Package);
+			UniqueLoadedPackages.Add(Package);
+		}
+		else
+		{
+			NotFoundPackages.Add(PackageName);
 		}
 
 		return false; // do not remove the package
@@ -999,7 +1003,49 @@ bool USourceControlHelpers::ApplyOperationAndReloadPackages(const TArray<FString
 		return false;
 	}
 
+	// Reverting may reintroduce some packages, so we need to ensure they're picked up...
+	if (!NotFoundPackages.IsEmpty() && bReloadWorld)
+	{
+		const FString& ExternalActorsFolderName = FPackagePath::GetExternalActorsFolderName();
+
+		// Gather the ExternalActorPaths in which we're reintroducing packages...
+		TSet<FString> NotFoundExternalPaths;
+		for (const FString& PackageName : NotFoundPackages)
+		{
+			int32 Index = PackageName.Find(ExternalActorsFolderName);
+			if (Index != INDEX_NONE)
+			{
+				// Format: /<MountPoint>/__ExternalActors__/<Level>/0/AB/CDEFGHIJKLMNOPQRSTUVWX
+				// Result: /<MountPoint>/__ExternalActors__/<Level>
+
+				Index += ExternalActorsFolderName.Len();
+				Index += 1;
+
+				Index = PackageName.Find("/", ESearchCase::IgnoreCase, ESearchDir::FromStart, Index);
+				if (Index != INDEX_NONE)
+				{
+					NotFoundExternalPaths.Add(PackageName.Left(Index));
+				}
+			}
+		}
+
+		// See if any of the loaded levels stores their external actors in any of those paths...
+		for (TObjectIterator<ULevel> LevelIt; LevelIt; ++LevelIt)
+		{
+			ULevel* Level = (*LevelIt);
+			if (Level->IsUsingExternalActors())
+			{
+				const FString& ExternalActorPath = ULevel::GetExternalActorsPath(Level->GetPackage());
+				if (NotFoundExternalPaths.Contains(ExternalActorPath))
+				{
+					UniqueLoadedPackages.Add(Level->GetWorld()->GetPackage());
+				}
+			}
+		}
+	}
+
 	// Prepare the packages to be reverted...
+	TArray<UPackage*> LoadedPackages = UniqueLoadedPackages.Array();
 	for (UPackage* Package : LoadedPackages)
 	{
 		// Detach the linkers of any loaded packages so that SCC can overwrite the files...
@@ -1029,7 +1075,7 @@ bool USourceControlHelpers::ApplyOperationAndReloadPackages(const TArray<FString
 				{
 					ObjectsToDelete.Add(ObjectToDelete);
 				}
-			}			
+			}
 			return true; // remove package
 		}
 		return false; // keep package
@@ -1044,7 +1090,7 @@ bool USourceControlHelpers::ApplyOperationAndReloadPackages(const TArray<FString
 	}
 
 	// Delete and Unload assets...
-	if(ObjectTools::DeleteObjectsUnchecked(ObjectsToDelete) != ObjectsToDelete.Num())
+	if (ObjectTools::DeleteObjectsUnchecked(ObjectsToDelete) != ObjectsToDelete.Num())
 	{ 
 		UE_LOG(LogSourceControl, Warning, TEXT("Failed to unload some assets."));
 	}
