SourceControl: Added cvar to force the modal dialog to always show up as a modal dialog, ignoring the preference of the calling code.

diff --git a/Engine/Source/Developer/SourceControl/Private/SourceControlCVars.cpp b/Engine/Source/Developer/SourceControl/Private/SourceControlCVars.cpp
--- a/Engine/Source/Developer/SourceControl/Private/SourceControlCVars.cpp
+++ b/Engine/Source/Developer/SourceControl/Private/SourceControlCVars.cpp
@@ -15,4 +15,10 @@ namespace SourceControlCVars
 		false,
 		TEXT("Allows a SourceControl 'Revert' operation to be triggered from the SubmitWidget."),
 		ECVF_Default);
+
+	TAutoConsoleVariable<bool> CVarSourceControlEnableLoginDialogModal(
+		TEXT("SourceControl.LoginDialog.ForceModal"),
+		false,
+		TEXT("Forces the SourceControl 'Login Dialog' to always be a modal dialog."),
+		ECVF_Default);
 }
\ No newline at end of file
diff --git a/Engine/Source/Developer/SourceControl/Private/SourceControlCVars.h b/Engine/Source/Developer/SourceControl/Private/SourceControlCVars.h
--- a/Engine/Source/Developer/SourceControl/Private/SourceControlCVars.h
+++ b/Engine/Source/Developer/SourceControl/Private/SourceControlCVars.h
@@ -9,4 +9,5 @@ namespace SourceControlCVars
 {
 	extern TAutoConsoleVariable<bool> CVarSourceControlEnableRevertFromSceneOutliner;
 	extern TAutoConsoleVariable<bool> CVarSourceControlEnableRevertFromSubmitWidget;
+	extern TAutoConsoleVariable<bool> CVarSourceControlEnableLoginDialogModal;
 }
\ No newline at end of file
diff --git a/Engine/Source/Developer/SourceControl/Private/SourceControlModule.cpp b/Engine/Source/Developer/SourceControl/Private/SourceControlModule.cpp
--- a/Engine/Source/Developer/SourceControl/Private/SourceControlModule.cpp
+++ b/Engine/Source/Developer/SourceControl/Private/SourceControlModule.cpp
@@ -7,6 +7,7 @@
 #include "SourceControlHelpers.h"
 #include "SourceControlAssetDataCache.h"
 #include "SourceControlFileStatusMonitor.h"
+#include "SourceControlCVars.h"
 #include "Misc/Paths.h"
 
 #if SOURCE_CONTROL_WITH_SLATE
@@ -109,8 +110,15 @@ void FSourceControlModule::ShowLoginDialog(const FSourceControlLoginClosed& InOn
 	// Get Active Provider Name
 	ActiveProviderName = GetProvider().GetName().ToString();
 
+	// if we are forcing a modal dialog, change the preference now
+	ELoginWindowMode::Type LoginWindowMode = InLoginWindowMode;
+	if(SourceControlCVars::CVarSourceControlEnableLoginDialogModal.GetValueOnAnyThread())
+	{
+		LoginWindowMode = ELoginWindowMode::Modal;
+	}
+
 	// if we are showing a modal version of the dialog & a modeless version already exists, we must destroy the modeless dialog first
-	if(InLoginWindowMode == ELoginWindowMode::Modal && SourceControlLoginPtr.IsValid())
+	if(LoginWindowMode == ELoginWindowMode::Modal && SourceControlLoginPtr.IsValid())
 	{
 		// unhook the delegate so it doesn't fire in this case
 		SourceControlLoginWindowPtr->SetOnWindowClosed(FOnWindowClosed());
@@ -157,7 +165,7 @@ void FSourceControlModule::ShowLoginDialog(const FSourceControlLoginClosed& InOn
 		TSharedPtr<SWindow> RootWindow = FGlobalTabmanager::Get()->GetRootWindow();
 		if(RootWindow.IsValid())
 		{
-			if(InLoginWindowMode == ELoginWindowMode::Modal)
+			if(LoginWindowMode == ELoginWindowMode::Modal)
 			{
 				FSlateApplication::Get().AddModalWindow(SourceControlLoginWindowPtr.ToSharedRef(), RootWindow);
 			}
@@ -168,7 +176,7 @@ void FSourceControlModule::ShowLoginDialog(const FSourceControlLoginClosed& InOn
 		}
 		else
 		{
-			if(InLoginWindowMode == ELoginWindowMode::Modal)
+			if(LoginWindowMode == ELoginWindowMode::Modal)
 			{
 				FSlateApplication::Get().AddModalWindow(SourceControlLoginWindowPtr.ToSharedRef(), RootWindow);
 			}
