Adding a modified icon for revision control, icon is a grey checkmark with a dot

#rnx

diff --git a/Engine/Source/Developer/SourceControl/Private/RevisionControlStyle/RevisionControlStyle.cpp b/Engine/Source/Developer/SourceControl/Private/RevisionControlStyle/RevisionControlStyle.cpp
--- a/Engine/Source/Developer/SourceControl/Private/RevisionControlStyle/RevisionControlStyle.cpp
+++ b/Engine/Source/Developer/SourceControl/Private/RevisionControlStyle/RevisionControlStyle.cpp
@@ -110,6 +110,7 @@ FDefaultRevisionControlStyle::FDefaultRevisionControlStyle() : FSlateStyleSet(St
 
 	Set("RevisionControl.ModifiedOtherBranch", new IMAGE_BRUSH_SVG("Starship/SourceControl/SCC_ModifiedOtherBranch", CoreStyleConstants::Icon16x16, BranchedColor));
 	Set("RevisionControl.ModifiedBadge", new IMAGE_BRUSH_SVG("Starship/SourceControl/RC_BranchModifiedBadge", CoreStyleConstants::Icon16x16, FStyleColors::Warning));
+	Set("RevisionControl.ModifiedLocally", new IMAGE_BRUSH_SVG("Starship/SourceControl/RC_ModifiedLocally", CoreStyleConstants::Icon16x16, FStyleColors::Foreground));
 
 	Set("RevisionControl.NotAtHeadRevision", new IMAGE_BRUSH_SVG("Starship/SourceControl/RC_NewerVersion", CoreStyleConstants::Icon16x16, FStyleColors::Warning));
 	Set("RevisionControl.NotInDepot", new IMAGE_BRUSH_SVG("Starship/SourceControl/SCC_NotInDepot", CoreStyleConstants::Icon16x16, FStyleColors::Warning));
