Removed redundant check for status of all files after a Revert operation

diff --git a/Engine/Source/Developer/SourceControl/Private/SourceControlHelpers.cpp b/Engine/Source/Developer/SourceControl/Private/SourceControlHelpers.cpp
--- a/Engine/Source/Developer/SourceControl/Private/SourceControlHelpers.cpp
+++ b/Engine/Source/Developer/SourceControl/Private/SourceControlHelpers.cpp
@@ -1096,17 +1096,7 @@ bool USourceControlHelpers::ApplyOperationAndReloadPackages(const TArray<FString
 	}
 	
 	// Re-cache the SCC state...
-	if (bReloadWorld)
-	{
-		auto UpdateCommand = ISourceControlOperation::Create<FUpdateStatus>();
-		UpdateCommand->SetCheckingAllFiles(true);
-		UpdateCommand->SetForceUpdate(true);
-		SourceControlProvider.Execute(UpdateCommand, EConcurrency::Asynchronous);
-	}
-	else
-	{
-		SourceControlProvider.Execute(ISourceControlOperation::Create<FUpdateStatus>(), PackageFilenames, EConcurrency::Asynchronous);
-	}
+	SourceControlProvider.Execute(ISourceControlOperation::Create<FUpdateStatus>(), PackageFilenames, EConcurrency::Asynchronous);
 
 	return bSuccess;
 }
