Allow additional error info to be stored in FSourceControlResultInfo.

- Added FSourceControlResultInfo::EAdditionalErrorContext that describes additional info about errors that the operation might have encountered that the caller can poll for rather than trying to parse error strings.
-- NOTE that we cannot enforce this info to be provided by source control provider implementations so this additional info should only be used to improve error messages.
-- Initially we support the connection to the provider failing and the connection to the provider dropping mid operation.
- Moved FSourceControlResultInfo to its own code files.
- I wanted to move the implementation of FSourceControlResultInfo::Append to the cpp, but this causes linker problems even with SOURCECONTROL_API as some modules are not including source control properly. For now I have left it in the header file and will look into this in another code pass.

diff --git a/Engine/Plugins/Developer/PerforceSourceControl/Source/PerforceSourceControl/Private/PerforceSourceControlProvider.cpp b/Engine/Plugins/Developer/PerforceSourceControl/Source/PerforceSourceControl/Private/PerforceSourceControlProvider.cpp
--- a/Engine/Plugins/Developer/PerforceSourceControl/Source/PerforceSourceControl/Private/PerforceSourceControlProvider.cpp
+++ b/Engine/Plugins/Developer/PerforceSourceControl/Source/PerforceSourceControl/Private/PerforceSourceControlProvider.cpp
@@ -767,6 +767,16 @@ bool FPerforceSourceControlProvider::TryToDownloadFileFromBackgroundThread(const
 	// Sanity check to make sure we are not running a command that modifies the cached states from a background thread
 	check(Command.Worker->UpdateStates() == false);
 
+	if (!Command.bConnectionWasSuccessful && !Command.bCancelled)
+	{
+		Command.ResultInfo.OnConnectionFailed();
+		
+	}
+	else if (Command.bConnectionDropped)
+	{
+		Command.ResultInfo.OnConnectionDroped();
+	}
+
 	OutputCommandMessages(Command);
 
 	if (!Command.bCancelledWhileTryingToConnect)
diff --git a/Engine/Source/Developer/SourceControl/Private/SourceControlResultInfo.cpp b/Engine/Source/Developer/SourceControl/Private/SourceControlResultInfo.cpp
new file mode 100644
--- /dev/null
+++ b/Engine/Source/Developer/SourceControl/Private/SourceControlResultInfo.cpp
@@ -0,0 +1,27 @@
+// Copyright Epic Games, Inc. All Rights Reserved.
+
+#include "SourceControlResultInfo.h"
+
+#include "Internationalization/Internationalization.h"
+
+#define LOCTEXT_NAMESPACE "SourceControlResultInfo"
+
+void FSourceControlResultInfo::OnConnectionFailed()
+{
+	ErrorMessages.Add(LOCTEXT("ConnectionFailed", "Failed to connect to the server"));
+	AdditionalErrorContext |= EAdditionalErrorContext::ConnectionFailed;
+}
+
+void FSourceControlResultInfo::OnConnectionDroped()
+{
+	ErrorMessages.Add(LOCTEXT("ConnectionFailed", "Connection to the server dropped"));
+	AdditionalErrorContext |= EAdditionalErrorContext::ConnectionDropped;
+}
+
+bool FSourceControlResultInfo::DidConnectionFail() const
+{
+	return EnumHasAnyFlags(AdditionalErrorContext,	EAdditionalErrorContext::ConnectionFailed |
+													EAdditionalErrorContext::ConnectionDropped);
+}
+
+#undef LOCTEXT_NAMESPACE
diff --git a/Engine/Source/Developer/SourceControl/Public/ISourceControlOperation.h b/Engine/Source/Developer/SourceControl/Public/ISourceControlOperation.h
--- a/Engine/Source/Developer/SourceControl/Public/ISourceControlOperation.h
+++ b/Engine/Source/Developer/SourceControl/Public/ISourceControlOperation.h
@@ -3,29 +3,7 @@
 #pragma once
 
 #include "CoreMinimal.h"
-
-
-/** Accumulated error and info messages for a source control operation.  */
-struct FSourceControlResultInfo
-{
-	/** Append any messages from another FSourceControlResultInfo, ensuring to keep any already accumulated info. */
-	void Append(const FSourceControlResultInfo& InResultInfo)
-	{
-		InfoMessages.Append(InResultInfo.InfoMessages);
-		ErrorMessages.Append(InResultInfo.ErrorMessages);
-		Tags.Append(InResultInfo.Tags);
-	}
-
-	/** Info and/or warning message storage */
-	TArray<FText> InfoMessages;
-
-	/** Potential error message storage */
-	TArray<FText> ErrorMessages;
-
-	/** Additional arbitrary information attached to the command */
-	TArray<FString> Tags;
-};
-
+#include "SourceControlResultInfo.h"
 
 class ISourceControlOperation : public TSharedFromThis<ISourceControlOperation, ESPMode::ThreadSafe>
 {
diff --git a/Engine/Source/Developer/SourceControl/Public/SourceControlResultInfo.h b/Engine/Source/Developer/SourceControl/Public/SourceControlResultInfo.h
new file mode 100644
--- /dev/null
+++ b/Engine/Source/Developer/SourceControl/Public/SourceControlResultInfo.h
@@ -0,0 +1,54 @@
+// Copyright Epic Games, Inc. All Rights Reserved.
+
+#pragma once
+
+#include "Containers/Array.h"
+#include "Containers/UnrealString.h"
+#include "CoreTypes.h"
+#include "Internationalization/Text.h"
+#include "Misc/EnumClassFlags.h"
+
+/** Accumulated error and info messages for a source control operation.  */
+struct FSourceControlResultInfo
+{
+	enum class EAdditionalErrorContext : uint32
+	{
+		None				= 0,
+		/** A connection to the server could not be established */
+		ConnectionFailed	= 1 << 0,
+		/** The connection to the server dropped before the operation could be completed */
+		ConnectionDropped	= 1 << 1
+	};
+
+	FRIEND_ENUM_CLASS_FLAGS(EAdditionalErrorContext);
+
+	/** Append any messages from another FSourceControlResultInfo, ensuring to keep any already accumulated info. */
+	void Append(const FSourceControlResultInfo& InResultInfo)
+	{
+		InfoMessages.Append(InResultInfo.InfoMessages);
+		ErrorMessages.Append(InResultInfo.ErrorMessages);
+		Tags.Append(InResultInfo.Tags);
+
+		AdditionalErrorContext |= InResultInfo.AdditionalErrorContext;
+	}
+
+	SOURCECONTROL_API void OnConnectionFailed();
+	SOURCECONTROL_API void OnConnectionDroped();
+
+	SOURCECONTROL_API bool DidConnectionFail() const;
+
+	/** Info and/or warning message storage */
+	TArray<FText> InfoMessages;
+
+	/** Potential error message storage */
+	TArray<FText> ErrorMessages;
+
+	/** Additional arbitrary information attached to the command */
+	TArray<FString> Tags;
+
+private:
+	/** Contains additional info about any errors encountered */
+	EAdditionalErrorContext AdditionalErrorContext = EAdditionalErrorContext::None;
+};
+
+ENUM_CLASS_FLAGS(FSourceControlResultInfo::EAdditionalErrorContext);
