CIS - Fix conflicting localization macros

#rnx

diff --git a/Engine/Source/Developer/SourceControl/Private/SourceControlResultInfo.cpp b/Engine/Source/Developer/SourceControl/Private/SourceControlResultInfo.cpp
--- a/Engine/Source/Developer/SourceControl/Private/SourceControlResultInfo.cpp
+++ b/Engine/Source/Developer/SourceControl/Private/SourceControlResultInfo.cpp
@@ -8,13 +8,13 @@
 
 void FSourceControlResultInfo::OnConnectionFailed()
 {
-	ErrorMessages.Add(LOCTEXT("ConnectionFailed", "Failed to connect to the server"));
+	ErrorMessages.Add(LOCTEXT("SC_ConnectionFailed", "Failed to connect to the server"));
 	AdditionalErrorContext |= EAdditionalErrorContext::ConnectionFailed;
 }
 
 void FSourceControlResultInfo::OnConnectionDroped()
 {
-	ErrorMessages.Add(LOCTEXT("ConnectionFailed", "Connection to the server dropped"));
+	ErrorMessages.Add(LOCTEXT("SC_ConnectionDropped", "Connection to the server dropped"));
 	AdditionalErrorContext |= EAdditionalErrorContext::ConnectionDropped;
 }
 
