Clean up of FDownloadFile

#rnx

- Added code documentation
- Renamed ::AddFileData to ::__Internal_AddFileData and documented that it is for internal use in the source control api. This is a standard problem with our operations being both input and output structures and needs a better solution.
- Moved ::GetFileData to the cpp to reduce the amount of code in the header file.

diff --git a/Engine/Plugins/Developer/PerforceSourceControl/Source/PerforceSourceControl/Private/PerforceSourceControlOperations.cpp b/Engine/Plugins/Developer/PerforceSourceControl/Source/PerforceSourceControl/Private/PerforceSourceControlOperations.cpp
--- a/Engine/Plugins/Developer/PerforceSourceControl/Source/PerforceSourceControl/Private/PerforceSourceControlOperations.cpp
+++ b/Engine/Plugins/Developer/PerforceSourceControl/Source/PerforceSourceControl/Private/PerforceSourceControlOperations.cpp
@@ -3095,7 +3095,7 @@ bool FPerforceDownloadFileWorker::Execute(FPerforceSourceControlCommand& InComma
 
 				for (int32 Index = 0; Index < InCommand.Files.Num(); ++Index)
 				{
-					Operation->AddFileData(InCommand.Files[Index], FilesData[Index]);
+					Operation->__Internal_AddFileData(InCommand.Files[Index], FilesData[Index]);
 				}
 			}
 		}
diff --git a/Engine/Source/Developer/SourceControl/Private/SourceControlOperations.cpp b/Engine/Source/Developer/SourceControl/Private/SourceControlOperations.cpp
--- a/Engine/Source/Developer/SourceControl/Private/SourceControlOperations.cpp
+++ b/Engine/Source/Developer/SourceControl/Private/SourceControlOperations.cpp
@@ -62,6 +62,20 @@ FDownloadFile::FDownloadFile(FStringView InTargetDirectory, EVerbosity InVerbosi
 	}
 }
 
+FSharedBuffer FDownloadFile::GetFileData(const FStringView& Filename)
+{
+	const uint32 Hash = GetTypeHash(Filename);
+	FSharedBuffer* Buffer = FileDataMap.FindByHash(Hash, Filename);
+	if (Buffer != nullptr)
+	{
+		return *Buffer;
+	}
+	else
+	{
+		return FSharedBuffer();
+	}
+}
+
 FCreateWorkspace::FCreateWorkspace(FStringView InWorkspaceName, FStringView InWorkspaceRoot)
 	: WorkspaceName(InWorkspaceName)
 {
diff --git a/Engine/Source/Developer/SourceControl/Public/SourceControlOperations.h b/Engine/Source/Developer/SourceControl/Public/SourceControlOperations.h
--- a/Engine/Source/Developer/SourceControl/Public/SourceControlOperations.h
+++ b/Engine/Source/Developer/SourceControl/Public/SourceControlOperations.h
@@ -1002,17 +1002,19 @@ public:
 	};
 
 	/** 
-	 * This constructor will download the files and keep them in memory, 
-	 * which can then be accessed by calling XXX.
+	 * This version of FDownloadFile will download the files and keep them in memory,
+	 * which can then be accessed by calling ::GetFileData. Full logging verbosity
+	 * will be used.
 	 */
 	FDownloadFile() = default;
 
-	FDownloadFile( EVerbosity InVerbosity)
+	/**
+	 * This version of FDownloadFile will download the files and keep them in memory,
+	 * which can then be accessed by calling ::GetFileData.
+	 */
+	FDownloadFile(EVerbosity InVerbosity)
 		: Verbosity(InVerbosity)
 	{}
-	
-	/** This constructor will download the files to the given directory */
-	SOURCECONTROL_API FDownloadFile(FStringView InTargetDirectory, EVerbosity InVerbosity);
 
 	// ISourceControlOperation interface
 	virtual FName GetName() const override
@@ -1020,45 +1022,61 @@ public:
 		return "DownloadFile";
 	}
 
+	// ISourceControlOperation interface
 	virtual FText GetInProgressString() const override
 	{
 		return LOCTEXT("SourceControl_PrintOperation", "Downloading file from server...");
 	}
 
-	virtual bool CanBeCalledFromBackgroundThreads() const
+	// ISourceControlOperation interface
+	virtual bool CanBeCalledFromBackgroundThreads() const override
 	{
 		return true;
 	}
+	
+	/** This version of FDownloadFile will download the files to the given target directory */
+	SOURCECONTROL_API FDownloadFile(FStringView InTargetDirectory, EVerbosity InVerbosity);
 
+	/**
+	 * Returns the directory that the files will (or have) been downloaded to. This path
+	 * will be empty if no target directory was given, in which case the files can be
+	 * accessed via the ::GetFileData method.
+	 */
 	FString GetTargetDirectory() const
 	{
 		return TargetDirectory;
 	}
 
-	void AddFileData(const FString& Filename, FSharedBuffer FileData)
+	/**
+	 * If no target directory was given to download the files too then the command will keep
+	 * the files in memory which can be accessed via this method. If the file failed to download
+	 * then a null FSharedBuffer will be returned.
+	 */
+	SOURCECONTROL_API FSharedBuffer GetFileData(const FStringView& Filename);
+
+	/** Return true if the command should log its operations, otherwise false */
+	bool ShouldLogToStdOutput() const
 	{
-		FileDataMap.Add(Filename, FileData);
+		return Verbosity == EVerbosity::Full;
 	}
 
-	FSharedBuffer GetFileData(const FStringView& Filename)
+	/** 
+	 * Do not call outside of source control implementations. Used to add the file in memory
+	 * once downloaded so that the caller can get access to it.
+	 */
+	void __Internal_AddFileData(const FString& Filename, FSharedBuffer FileData)
 	{
-		const uint32 Hash = GetTypeHash(Filename);
-		FSharedBuffer* Buffer = FileDataMap.FindByHash(Hash, Filename);
-		if (Buffer != nullptr)
-		{
-			return *Buffer;
-		}
-		else
-		{
-			return FSharedBuffer();
-		}
+		FileDataMap.Add(Filename, FileData);
 	}
 
-	bool ShouldLogToStdOutput() const
+	UE_DEPRECATED(5.3, "Replaced by __Internal_AddFileData but you shouldn't be calling this anyway")
+	void AddFileData(const FString& Filename, FSharedBuffer FileData)
 	{
-		return Verbosity == EVerbosity::Full;
+		__Internal_AddFileData(Filename, FileData);
 	}
+
 private:
+
 	EVerbosity Verbosity = EVerbosity::Full;
 
 	FString TargetDirectory;
