Fixed: Unable to revert level file deletion in Check-in Window

diff --git a/Engine/Source/Developer/SourceControl/Private/SourceControlHelpers.cpp b/Engine/Source/Developer/SourceControl/Private/SourceControlHelpers.cpp
--- a/Engine/Source/Developer/SourceControl/Private/SourceControlHelpers.cpp
+++ b/Engine/Source/Developer/SourceControl/Private/SourceControlHelpers.cpp
@@ -917,6 +917,13 @@ bool USourceControlHelpers::ApplyOperationAndReloadPackages(const TArray<FString
 	TArray<FString> PackageFilenames;
 	TArray<FString> FilteredActorPackages;
 	TArray<FString> NotFoundPackages;
+	/** 
+	 * This is necessary to cover an edge case where the user reverts a deleted Level file.
+	 * In that case PackageFilename_Internal() will not be able to resolve to the correct extension
+	 * as the file is deleted and the package cannot be inspected
+	 */
+	TMap<FString, FString> MapPackageNamesToFilenames;
+
 	bool bSuccess = false;
 
 	auto DetachLinker = [](UPackage* Package)
@@ -936,6 +943,12 @@ bool USourceControlHelpers::ApplyOperationAndReloadPackages(const TArray<FString
 		
 		if (FPackageName::TryConvertFilenameToLongPackageName(Filename, Result))
 		{
+			FString Extension = FPaths::GetExtension(Filename);
+			if (Extension == TEXT("umap"))
+			{
+				MapPackageNamesToFilenames.Add(Result, Filename);
+			}
+
 			PackageNames.Add(MoveTemp(Result));
 		}
 		else
@@ -1052,7 +1065,17 @@ bool USourceControlHelpers::ApplyOperationAndReloadPackages(const TArray<FString
 		DetachLinker(Package);
 	}
 
-	PackageFilenames = SourceControlHelpers::PackageFilenames(PackageNames);
+	for (int32 PackageIndex = 0; PackageIndex < PackageNames.Num(); PackageIndex++)
+	{
+		if (MapPackageNamesToFilenames.Contains(PackageNames[PackageIndex]))
+		{
+			PackageFilenames.Add(MapPackageNamesToFilenames[PackageNames[PackageIndex]]);
+		}
+		else
+		{
+			PackageFilenames.Add(PackageFilename(PackageNames[PackageIndex]));
+		}
+	}
 
 	// Apply Operation
 	bSuccess = InOperation(PackageFilenames);
