Implemented the FGetFileList operation for URC

Discussed this with [REDACTED_HANDLE] - the `FGetFileList` operation is only used by `USourceControlHelpers::GetFilesInDepotAtPath()`, and that method is currently not called by anything.

Also changed the implementation of `GetFilesInDepotAtPath()` to pass the search pattern as a parameter to the Operation.

diff --git a/Engine/Source/Developer/SourceControl/Private/SourceControlHelpers.cpp b/Engine/Source/Developer/SourceControl/Private/SourceControlHelpers.cpp
--- a/Engine/Source/Developer/SourceControl/Private/SourceControlHelpers.cpp
+++ b/Engine/Source/Developer/SourceControl/Private/SourceControlHelpers.cpp
@@ -1641,9 +1641,10 @@ bool USourceControlHelpers::GetFilesInDepotAtPath(const FString& PathToDirectory
 
 			TArray<FString> FileArray;
 			FileArray.Add(FullPath);
-
+			
 			TSharedRef<FGetFileList, ESPMode::ThreadSafe> Operation = ISourceControlOperation::Create<FGetFileList>();
 			Operation->SetIncludeDeleted(bIncludeDeleted);
+			Operation->SetSearchPattern(PathToDirectory);
 
 			ECommandResult::Type Result = SourceControlProvider.Execute(Operation, FileArray, EConcurrency::Synchronous);
 			bSuccess = (Result == ECommandResult::Succeeded);
diff --git a/Engine/Source/Developer/SourceControl/Public/SourceControlOperations.h b/Engine/Source/Developer/SourceControl/Public/SourceControlOperations.h
--- a/Engine/Source/Developer/SourceControl/Public/SourceControlOperations.h
+++ b/Engine/Source/Developer/SourceControl/Public/SourceControlOperations.h
@@ -180,12 +180,25 @@ public:
 		FilesList = MoveTemp(InFilesList);
 	}
 
+	const FString& GetSearchPattern() const
+	{
+		return SearchPattern;
+	}
+
+	void SetSearchPattern(const FString& InSearchPattern)
+	{
+		SearchPattern = InSearchPattern;
+	}
+
 protected:
 	/** Include deleted files in the list. */
 	bool bIncludeDeleted = false;
 
 	/** Stored result of the operation */
 	TArray<FString> FilesList;
+
+	/** The search pattern for the file list */
+	FString SearchPattern;
 };
 
 /**
