Make sure to use loadedpath only for in memory files
Otherwise renamed packages will not resolve

diff --git a/Engine/Source/Developer/SourceControl/Private/SourceControlHelpers.cpp b/Engine/Source/Developer/SourceControl/Private/SourceControlHelpers.cpp
--- a/Engine/Source/Developer/SourceControl/Private/SourceControlHelpers.cpp
+++ b/Engine/Source/Developer/SourceControl/Private/SourceControlHelpers.cpp
@@ -1699,10 +1699,11 @@ FString USourceControlHelpers::PackageFilename( const UPackage* InPackage )
 	if(InPackage != nullptr)
 	{
 		// Prefer using package loaded path to resolve file name as it properly resolves memory packages
-		if (!InPackage->GetLoadedPath().IsEmpty())
+		FString PackageLoadedPath = InPackage->GetLoadedPath().GetPackageName();
+		if (!InPackage->GetLoadedPath().IsEmpty() && FPackageName::IsMemoryPackage(PackageLoadedPath))
 		{
 			const FString PackageExtension = InPackage->ContainsMap() ? FPackageName::GetMapPackageExtension() : FPackageName::GetAssetPackageExtension();
-			Filename = FPaths::ConvertRelativePathToFull(FPackageName::LongPackageNameToFilename(InPackage->GetLoadedPath().GetPackageName(), PackageExtension));
+			Filename = FPaths::ConvertRelativePathToFull(FPackageName::LongPackageNameToFilename(PackageLoadedPath, PackageExtension));
 		}
 		else
 		{
