SourceControl: Fixed actor label change not being reflected in the 'Submit Files' dialog because the cached FAssetData still contained the old label.

#rnx

diff --git a/Engine/Source/Developer/SourceControl/Private/SourceControlAssetDataCache.cpp b/Engine/Source/Developer/SourceControl/Private/SourceControlAssetDataCache.cpp
--- a/Engine/Source/Developer/SourceControl/Private/SourceControlAssetDataCache.cpp
+++ b/Engine/Source/Developer/SourceControl/Private/SourceControlAssetDataCache.cpp
@@ -5,12 +5,17 @@
 #include "Async/Async.h"
 #include "AssetRegistry/AssetData.h"
 #include "HAL/IConsoleManager.h"
+#include "Misc/CoreDelegates.h"
 #include "AssetRegistry/IAssetRegistry.h"
 #include "ISourceControlModule.h"
 #include "ISourceControlState.h"
 #include "SourceControlHelpers.h"
 #include "SourceControlOperations.h"
 
+#if WITH_EDITOR
+#include "GameFramework/Actor.h"
+#endif
+
 static TAutoConsoleVariable<int32> CVarSourceControlAssetDataCacheMaxAsyncTask(
 	TEXT("SourceControlAssetDataCache.MaxAsyncTask"),
 	8,
@@ -26,11 +31,29 @@ void FSourceControlAssetDataCache::Startup()
 		ISourceControlModule::Get().GetAssetDataCache().OnSourceControlProviderChanged(OldProvider, NewProvider);
 	}));
 
+#if WITH_EDITOR
+	ActorLabelChangedDelegateHandle = FCoreDelegates::OnActorLabelChanged.AddLambda([](AActor* ChangedActor)
+	{
+		if (ensure(ChangedActor))
+		{
+			if (UPackage* Package = ChangedActor->GetPackage())
+			{
+				FString Filename = USourceControlHelpers::PackageFilename(Package);
+				ISourceControlModule::Get().GetAssetDataCache().ClearAssetData(Filename);
+			}
+		}
+	});
+#endif
+
 	bIsSourceControlDialogShown = false;
 }
 
 void FSourceControlAssetDataCache::Shutdown()
 {
+#if WITH_EDITOR
+	FCoreDelegates::OnActorLabelChanged.Remove(ActorLabelChangedDelegateHandle);
+#endif
+
 	ISourceControlModule& SCCModule = ISourceControlModule::Get();
 	SCCModule.UnregisterProviderChanged(ProviderChangedDelegateHandle);
 
@@ -179,6 +202,17 @@ void FSourceControlAssetDataCache::OnUpdateHistoryComplete(const TArray<FString>
 	}
 }
 
+void FSourceControlAssetDataCache::ClearAssetData(const FString& Filename)
+{
+	if (FSourceControlAssetDataEntry* AssetDataEntry = AssetDataCache.Find(Filename))
+	{
+		if (AssetDataEntry->bInitialized)
+		{
+			AssetDataCache.Remove(Filename);
+		}
+	}
+}
+
 void FSourceControlAssetDataCache::LaunchFetchAssetDataTasks()
 {
 	FSourceControlStatePtr FileState;
diff --git a/Engine/Source/Developer/SourceControl/Public/SourceControlAssetDataCache.h b/Engine/Source/Developer/SourceControl/Public/SourceControlAssetDataCache.h
--- a/Engine/Source/Developer/SourceControl/Public/SourceControlAssetDataCache.h
+++ b/Engine/Source/Developer/SourceControl/Public/SourceControlAssetDataCache.h
@@ -103,12 +103,16 @@ private:
 	/** Wait for the currently running tasks to finish and clears all queues. */
 	void ClearPendingTasks();
 
+	/** Clears the asset data associated with it */
+	void ClearAssetData(const FString& Filename);
+
 private:
 	FAssetDataCache										AssetDataCache;
 	TArray<FString>										FileHistoryToUpdate;
 	TQueue<FSourceControlStatePtr, EQueueMode::Spsc>	AssetDataToFetch;
 	TQueue<FAssetDataToLoad, EQueueMode::Mpsc>			AssetDataToLoad;
 	FDelegateHandle										ProviderChangedDelegateHandle;
+	FDelegateHandle										ActorLabelChangedDelegateHandle;
 	std::atomic<uint32>									CurrentAsyncTask = 0;
 	bool												bIsSourceControlDialogShown;
 };
