SourceControl: Widen the 'transfer size' information from 'int32' to 'int64'.

#rnx

diff --git a/Engine/Source/Developer/SourceControl/Public/SourceControlOperations.h b/Engine/Source/Developer/SourceControl/Public/SourceControlOperations.h
--- a/Engine/Source/Developer/SourceControl/Public/SourceControlOperations.h
+++ b/Engine/Source/Developer/SourceControl/Public/SourceControlOperations.h
@@ -329,12 +329,12 @@ public:
 		return bHeadRevision;
 	}
 
-	void SetTransferSize(const int32 InTransferSize)
+	void SetTransferSize(const int64 InTransferSize)
 	{
 		TransferSize = InTransferSize;
 	}
 
-	int32 GetTransferSize() const
+	int64 GetTransferSize() const
 	{
 		return TransferSize;
 	}
@@ -357,7 +357,7 @@ protected:
 	bool bHeadRevision = false;
 
 	/** Number of bytes that need to be transferred for the sync operation */
-	int32 TransferSize = 0;
+	int64 TransferSize = 0;
 
 	/** Array of files that would be affected by the sync operation */
 	TArray<FString> AffectedFiles;
