Truncation warning fixes for SourceControl module

#rnx

diff --git a/Engine/Source/Developer/SourceControl/Private/SSourceControlLogin.cpp b/Engine/Source/Developer/SourceControl/Private/SSourceControlLogin.cpp
--- a/Engine/Source/Developer/SourceControl/Private/SSourceControlLogin.cpp
+++ b/Engine/Source/Developer/SourceControl/Private/SSourceControlLogin.cpp
@@ -108,7 +108,7 @@ void SSourceControlLogin::Construct(const FArguments& InArgs)
 						.Padding(0)
 						[
 							SNew(SBox)
-							.HeightOverride(250)
+							.HeightOverride(250.0f)
 							[
 								MessageLogModule.CreateLogListingWidget(MessageLogListing)
 							]
diff --git a/Engine/Source/Developer/SourceControl/Private/ScopedSourceControlProgress.cpp b/Engine/Source/Developer/SourceControl/Private/ScopedSourceControlProgress.cpp
--- a/Engine/Source/Developer/SourceControl/Private/ScopedSourceControlProgress.cpp
+++ b/Engine/Source/Developer/SourceControl/Private/ScopedSourceControlProgress.cpp
@@ -182,7 +182,7 @@ private:
 	bool bCancelClicked;
 
 	/** The timer we use to determine when to display the 'long task' message */
-	float TimeStamp;
+	double TimeStamp;
 
 	/** The delegate to call when the cancel button is clicked */
 	FSimpleDelegate OnCancelled;
diff --git a/Engine/Source/Developer/SourceControl/SourceControl.Build.cs b/Engine/Source/Developer/SourceControl/SourceControl.Build.cs
--- a/Engine/Source/Developer/SourceControl/SourceControl.Build.cs
+++ b/Engine/Source/Developer/SourceControl/SourceControl.Build.cs
@@ -27,10 +27,10 @@ public class SourceControl : ModuleRules
 			}
 		);
 
-        if (Target.bBuildEditor)
-        {
+		if (Target.bBuildEditor)
+		{
 			PrivateDependencyModuleNames.AddRange(
-                new string[] {
+				new string[] {
 					"EditorFramework",
 					"Engine",
 					"UnrealEd",
@@ -44,7 +44,7 @@ public class SourceControl : ModuleRules
 			);
 
 			CircularlyReferencedDependentModules.Add("UnrealEd");
-        }
+		}
 
 		if (Target.bBuildDeveloperTools)
 		{
@@ -59,5 +59,7 @@ public class SourceControl : ModuleRules
 		{
 			PrecompileForTargets = PrecompileTargetsType.Any;
 		}
+
+		UnsafeTypeCastWarningLevel = WarningLevel.Error;
 	}
 }
