SourceControl: Changes to FGetSubmittedChangelists (introduced rather plainly in #20292011 but unused since).
- Added support for returning results.
- Added support for filtering as it clearly doesn't scale without it.

#rnx

diff --git a/Engine/Source/Developer/SourceControl/Public/SourceControlOperations.h b/Engine/Source/Developer/SourceControl/Public/SourceControlOperations.h
--- a/Engine/Source/Developer/SourceControl/Public/SourceControlOperations.h
+++ b/Engine/Source/Developer/SourceControl/Public/SourceControlOperations.h
@@ -691,6 +691,62 @@ public:
 	{
 		return LOCTEXT("SourceControl_GetSubmittedChangelists", "Retrieving submitted changelist(s) from Revision Control...");
 	}
+
+public:
+	void SetDateFromFilter(const FDateTime& InDateFrom)
+	{
+		DateFrom = InDateFrom;
+	}
+
+	void SetDateToFilter(const FDateTime& InDateTo)
+	{
+		DateTo = InDateTo;
+	}
+
+	void SetOwnedFilter(bool InOwned)
+	{
+		Owned = InOwned;
+	}
+
+	bool ShouldFilterByDateFrom(FDateTime& OutDateFrom)
+	{
+		if (DateFrom != FDateTime::MinValue())
+		{
+			OutDateFrom = DateFrom;
+			return true;
+		}
+		return false;
+	}
+
+	bool ShouldFilterByDateTo(FDateTime& OutDateTo)
+	{
+		if (DateTo != FDateTime::MinValue())
+		{
+			OutDateTo = DateTo;
+			return true;
+		}
+		return false;
+	}
+
+	bool ShouldFilterByOwned()
+	{
+		return Owned;
+	}
+
+public:
+	const TArray<FSourceControlChangelistRef>& GetSubmittedChangelists() const { return Changelists; }
+
+	void AddSubmittedChangelist(FSourceControlChangelistRef Changelist)
+	{
+		Changelists.Add(Changelist);
+	}
+
+private:
+	FDateTime DateFrom = FDateTime::MinValue();
+	FDateTime DateTo = FDateTime::MinValue();
+	bool Owned = false;
+
+	TArray<FSourceControlChangelistRef> Changelists;
 };
 
 /**
