Extended RunProjectSnapshotList() method to support explicit pagination

diff --git a/Engine/Source/Developer/SourceControl/Public/SourceControlOperations.h b/Engine/Source/Developer/SourceControl/Public/SourceControlOperations.h
--- a/Engine/Source/Developer/SourceControl/Public/SourceControlOperations.h
+++ b/Engine/Source/Developer/SourceControl/Public/SourceControlOperations.h
@@ -708,6 +708,22 @@ public:
 		Owned = InOwned;
 	}
 
+	void SetPaginationLimit(int32 InLimit)
+	{
+		if (InLimit > 0)
+		{
+			PaginationLimit = InLimit;
+		}
+	}
+
+	void SetPaginationOffset(int32 InOffset)
+	{
+		if (InOffset >= 0)
+		{
+			PaginationOffset = InOffset;
+		}
+	}
+
 	bool ShouldFilterByDateFrom(FDateTime& OutDateFrom)
 	{
 		if (DateFrom != FDateTime::MinValue())
@@ -733,6 +749,26 @@ public:
 		return Owned;
 	}
 
+	bool ShouldUsePaginationLimit(int32& OutPaginationLimit)
+	{
+		if (PaginationLimit != -1)
+		{
+			OutPaginationLimit = PaginationLimit;
+			return true;
+		}
+		return false;
+	}
+
+	bool ShouldUsePaginationOffset(int32& OutPaginationOffset)
+	{
+		if (PaginationOffset != -1)
+		{
+			OutPaginationOffset = PaginationOffset;
+			return true;
+		}
+		return false;
+	}
+
 public:
 	const TArray<FSourceControlChangelistRef>& GetSubmittedChangelists() const { return Changelists; }
 
@@ -746,6 +782,9 @@ private:
 	FDateTime DateTo = FDateTime::MinValue();
 	bool Owned = false;
 
+	int32 PaginationLimit = -1;
+	int32 PaginationOffset = -1;
+
 	TArray<FSourceControlChangelistRef> Changelists;
 };
 
