SourceControl: Fixed world not reloading when a deleted external object is reverted (when it is not in memory anymore).

The existing code was specific to external Actors, but there are also external Objects, such as the folders created in the SceneOutliner. Previously, deleting a folder in the SceneOutliner, saving it, restarting the editor and reverting it would not trigger a world reload, incorrectly not showing the folder in the SceneOutliner.

#rnx

diff --git a/Engine/Source/Developer/SourceControl/Private/SourceControlHelpers.cpp b/Engine/Source/Developer/SourceControl/Private/SourceControlHelpers.cpp
--- a/Engine/Source/Developer/SourceControl/Private/SourceControlHelpers.cpp
+++ b/Engine/Source/Developer/SourceControl/Private/SourceControlHelpers.cpp
@@ -1020,24 +1020,32 @@ bool USourceControlHelpers::ApplyOperationAndReloadPackages(const TArray<FString
 	if (!NotFoundPackages.IsEmpty() && bReloadWorld)
 	{
 		const FString& ExternalActorsFolderName = FPackagePath::GetExternalActorsFolderName();
+		const FString& ExternalObjectsFolderName = FPackagePath::GetExternalObjectsFolderName();
 
-		// Gather the ExternalActorPaths in which we're reintroducing packages...
+		TArray<FString> ExternalFolderNames;
+		ExternalFolderNames.Add(ExternalActorsFolderName);
+		ExternalFolderNames.Add(ExternalObjectsFolderName);
+
+		// Gather the ExternalPaths in which we're reintroducing packages...
 		TSet<FString> NotFoundExternalPaths;
-		for (const FString& PackageName : NotFoundPackages)
+		for (const FString& ExternalFolderName : ExternalFolderNames)
 		{
-			int32 Index = PackageName.Find(ExternalActorsFolderName);
-			if (Index != INDEX_NONE)
+			for (const FString& PackageName : NotFoundPackages)
 			{
-				// Format: /<MountPoint>/__ExternalActors__/<Level>/0/AB/CDEFGHIJKLMNOPQRSTUVWX
-				// Result: /<MountPoint>/__ExternalActors__/<Level>
-
-				Index += ExternalActorsFolderName.Len();
-				Index += 1;
-
-				Index = PackageName.Find("/", ESearchCase::IgnoreCase, ESearchDir::FromStart, Index);
+				int32 Index = PackageName.Find(ExternalFolderName);
 				if (Index != INDEX_NONE)
 				{
-					NotFoundExternalPaths.Add(PackageName.Left(Index));
+					// Format: /<MountPoint>/__External<xxxxx>__/<Level>/0/AB/CDEFGHIJKLMNOPQRSTUVWX
+					// Result: /<MountPoint>/__External<xxxxx>__/<Level>
+
+					Index += ExternalFolderName.Len();
+					Index += 1;
+
+					Index = PackageName.Find("/", ESearchCase::IgnoreCase, ESearchDir::FromStart, Index);
+					if (Index != INDEX_NONE)
+					{
+						NotFoundExternalPaths.Add(PackageName.Left(Index));
+					}
 				}
 			}
 		}
@@ -1054,6 +1062,17 @@ bool USourceControlHelpers::ApplyOperationAndReloadPackages(const TArray<FString
 					UniqueLoadedPackages.Add(Level->GetWorld()->GetPackage());
 				}
 			}
+			if (Level->IsUsingExternalObjects())
+			{
+				const TArray<FString> ExternalObjectPaths = ULevel::GetExternalObjectsPaths(Level->GetPackage()->GetName());
+				for (const FString& ExternalObjectPath : ExternalObjectPaths)
+				{
+					if (NotFoundExternalPaths.Contains(ExternalObjectPath))
+					{
+						UniqueLoadedPackages.Add(Level->GetWorld()->GetPackage());
+					}
+				}
+			}
 		}
 	}
 
