SourceControl: Fixed world not reloading when an external object is reverted.

The existing code was specific to external Actors, but there are also external Objects, such as the folders created in the SceneOutliner. Previously, creating a folder in the SceneOutliner, saving it and reverting it would not trigger a world reload, incorrectly keeping the folder visible in the SceneOutliner.

#rnx

diff --git a/Engine/Source/Developer/SourceControl/Private/SourceControlHelpers.cpp b/Engine/Source/Developer/SourceControl/Private/SourceControlHelpers.cpp
--- a/Engine/Source/Developer/SourceControl/Private/SourceControlHelpers.cpp
+++ b/Engine/Source/Developer/SourceControl/Private/SourceControlHelpers.cpp
@@ -915,7 +915,7 @@ bool USourceControlHelpers::ApplyOperationAndReloadPackages(const TArray<FString
 
 	TArray<FString> PackageNames;
 	TArray<FString> PackageFilenames;
-	TArray<FString> FilteredActorPackages;
+	TArray<FString> FilteredPackages;
 	TArray<FString> NotFoundPackages;
 	/** 
 	 * This is necessary to cover an edge case where the user reverts a deleted Level file.
@@ -957,42 +957,45 @@ bool USourceControlHelpers::ApplyOperationAndReloadPackages(const TArray<FString
 		}
 	}
 
-	// If bReloadWorld=false, remove packages if they are loaded actors or world
+	// If bReloadWorld=false, remove packages if they are loaded external packages or world
 	// If bReloadWorld=true, include world packages to reload
 	TSet<UPackage*> UniqueLoadedPackages;
-	PackageNames.RemoveAll([&FilteredActorPackages, &UniqueLoadedPackages, &NotFoundPackages, bReloadWorld, &DetachLinker](const FString& PackageName) -> bool
+	PackageNames.RemoveAll([&FilteredPackages, &UniqueLoadedPackages, &NotFoundPackages, bReloadWorld, &DetachLinker](const FString& PackageName) -> bool
 	{
 		UPackage* Package = FindPackage(NULL, *PackageName);
 		
 		if (Package != nullptr)
 		{
-			if (UWorld* World = UWorld::FindWorldInPackage(Package); World)
+			if (UWorld* World = UWorld::FindWorldInPackage(Package))
 			{
 				if (!bReloadWorld)
 				{
-					FilteredActorPackages.Emplace(PackageName);
+					FilteredPackages.Emplace(PackageName);
 					return true; // remove the package
 				}
 			}
-			else if (AActor* Actor = AActor::FindActorInPackage(Package))
+			else if (UObject* Asset = Package->FindAssetInPackage())
 			{
-				if (bReloadWorld)
+				if (Asset->IsPackageExternal())
 				{
-					// detach linker on the actor
-					DetachLinker(Package);
+					if (bReloadWorld)
+					{
+						// detach linker on the object
+						DetachLinker(Package);
+
+						// but track its world for reloading - not the object package itself
+						if (Asset->GetWorld() && Asset->GetWorld()->GetPackage())
+						{
+							UniqueLoadedPackages.Add(Asset->GetWorld()->GetPackage());
+						}
 
-					// but track its world for reloading - not the actor package itself
-					if (Actor->GetWorld() && Actor->GetWorld()->GetPackage())
+						return false;
+					}
+					else
 					{
-						UniqueLoadedPackages.Add(Actor->GetWorld()->GetPackage());
+						FilteredPackages.Emplace(PackageName);
+						return true; // remove the package
 					}
-
-					return false;
-				}
-				else if (Actor->IsPackageExternal())
-				{
-					FilteredActorPackages.Emplace(PackageName);
-					return true; // remove the package
 				}
 			}
 
@@ -1006,10 +1009,10 @@ bool USourceControlHelpers::ApplyOperationAndReloadPackages(const TArray<FString
 		return false; // do not remove the package
 	});
 
-	if (!FilteredActorPackages.IsEmpty())
+	if (!FilteredPackages.IsEmpty())
 	{
 		TStringBuilder<2048> Builder;
-		Builder.Join(FilteredActorPackages, TEXT(", "));
+		Builder.Join(FilteredPackages, TEXT(", "));
 		const FString Packages = Builder.ToString();
 
 		UE_LOG(LogSourceControl, Warning, TEXT("This operation could not complete on the following map or external packages, please unload them before retrying : %s"), *Packages);
