SourceControl: Fixed crash when reverting a newly added level file along with a world reload that that causes the editor to crash.

#rnx

diff --git a/Engine/Source/Developer/SourceControl/Private/SourceControlHelpers.cpp b/Engine/Source/Developer/SourceControl/Private/SourceControlHelpers.cpp
--- a/Engine/Source/Developer/SourceControl/Private/SourceControlHelpers.cpp
+++ b/Engine/Source/Developer/SourceControl/Private/SourceControlHelpers.cpp
@@ -1104,40 +1104,59 @@ bool USourceControlHelpers::ApplyOperationAndReloadPackages(const TArray<FString
 
 	// Reverting may have deleted some packages, so we need to delete those and unload them rather than re-load them...
 	FAssetRegistryModule& AssetRegistryModule = FModuleManager::LoadModuleChecked<FAssetRegistryModule>(TEXT("AssetRegistry"));
-	TArray<UObject*> ObjectsToDelete;
+	TArray<TWeakObjectPtr<UObject>> ObjectsMissingOnDisk;
 	LoadedPackages.RemoveAll([&](UPackage* InPackage) -> bool
-	{
-		const FString PackageExtension = InPackage->ContainsMap() ? FPackageName::GetMapPackageExtension() : FPackageName::GetAssetPackageExtension();
-		const FString PackageFilename = FPackageName::LongPackageNameToFilename(InPackage->GetName(), PackageExtension);
-		if (!FPaths::FileExists(PackageFilename))
 		{
-			TArray<FAssetData> Assets;
-			AssetRegistryModule.Get().GetAssetsByPackageName(*InPackage->GetName(), Assets);
-
-			for (const FAssetData& Asset : Assets)
+			const FString PackageExtension = InPackage->ContainsMap() ? FPackageName::GetMapPackageExtension() : FPackageName::GetAssetPackageExtension();
+			const FString PackageFilename = FPackageName::LongPackageNameToFilename(InPackage->GetName(), PackageExtension);
+			if (!FPaths::FileExists(PackageFilename))
 			{
-				if (UObject* ObjectToDelete = Asset.FastGetAsset())
+				TArray<FAssetData> Assets;
+				AssetRegistryModule.Get().GetAssetsByPackageName(*InPackage->GetName(), Assets);
+
+				for (const FAssetData& Asset : Assets)
 				{
-					ObjectsToDelete.Add(ObjectToDelete);
+					if (UObject* ObjectToDelete = Asset.FastGetAsset())
+					{
+						ObjectsMissingOnDisk.Add(ObjectToDelete);
+					}
 				}
+				return true; // remove package
 			}
-			return true; // remove package
-		}
-		return false; // keep package
-	});
+			return false; // keep package
+		});
 
 	// Hot-reload the new packages...
-	FText OutReloadErrorMsg;
-	UPackageTools::ReloadPackages(LoadedPackages, OutReloadErrorMsg, bInteractive ? EReloadPackagesInteractionMode::Interactive : EReloadPackagesInteractionMode::AssumePositive);
-	if (!OutReloadErrorMsg.IsEmpty())
+	if (LoadedPackages.Num() > 0)
 	{
-		UE_LOG(LogSourceControl, Warning, TEXT("%s"), *OutReloadErrorMsg.ToString());
+		FText OutReloadErrorMsg;
+		UPackageTools::ReloadPackages(LoadedPackages, OutReloadErrorMsg, bInteractive ? EReloadPackagesInteractionMode::Interactive : EReloadPackagesInteractionMode::AssumePositive);
+		if (!OutReloadErrorMsg.IsEmpty())
+		{
+			UE_LOG(LogSourceControl, Warning, TEXT("%s"), *OutReloadErrorMsg.ToString());
+		}
+	}
+
+	// A world reload might have already deleted some objects, so check which missing ones are still valid.
+	TArray<UObject*> ObjectsToDelete;
+	if (ObjectsMissingOnDisk.Num() > 0)
+	{
+		for (TWeakObjectPtr<UObject> Object : ObjectsMissingOnDisk)
+		{
+			if (UObject* ObjectToDelete = Object.Get())
+			{
+				ObjectsToDelete.Add(ObjectToDelete);
+			}
+		}
 	}
 
 	// Delete and Unload assets...
-	if (ObjectTools::DeleteObjectsUnchecked(ObjectsToDelete) != ObjectsToDelete.Num())
-	{ 
-		UE_LOG(LogSourceControl, Warning, TEXT("Failed to unload some assets."));
+	if (ObjectsToDelete.Num() > 0)
+	{
+		if (ObjectTools::DeleteObjectsUnchecked(ObjectsToDelete) != ObjectsToDelete.Num())
+		{
+			UE_LOG(LogSourceControl, Warning, TEXT("Failed to unload some assets."));
+		}
 	}
 	
 	// Re-cache the SCC state...
