Added #if SOURCE_CONTROL_WITH_SLATE to SSourceControlControls Slate class in SourceControl module

diff --git a/Engine/Source/Developer/SourceControl/Private/SSourceControlControls.cpp b/Engine/Source/Developer/SourceControl/Private/SSourceControlControls.cpp
--- a/Engine/Source/Developer/SourceControl/Private/SSourceControlControls.cpp
+++ b/Engine/Source/Developer/SourceControl/Private/SSourceControlControls.cpp
@@ -1,5 +1,7 @@
 // Copyright Epic Games, Inc. All Rights Reserved.
 
+#if SOURCE_CONTROL_WITH_SLATE
+
 #include "SSourceControlControls.h"
 #include "ISourceControlProvider.h"
 #include "ISourceControlModule.h"
@@ -389,4 +391,6 @@ FReply SSourceControlControls::OnSourceControlCheckInChangesClicked() const
 }
 
 
-#undef LOCTEXT_NAMESPACE
\ No newline at end of file
+#undef LOCTEXT_NAMESPACE
+
+#endif // SOURCE_CONTROL_WITH_SLATE
\ No newline at end of file
diff --git a/Engine/Source/Developer/SourceControl/Public/SSourceControlControls.h b/Engine/Source/Developer/SourceControl/Public/SSourceControlControls.h
--- a/Engine/Source/Developer/SourceControl/Public/SSourceControlControls.h
+++ b/Engine/Source/Developer/SourceControl/Public/SSourceControlControls.h
@@ -2,6 +2,8 @@
 
 #pragma once
 
+#if SOURCE_CONTROL_WITH_SLATE
+
 #include "CoreMinimal.h"
 #include "ISourceControlProvider.h"
 #include "Delegates/IDelegateInstance.h"
@@ -61,4 +63,6 @@ private:
 
 	FDelegateHandle SourceControlProviderChangedHandle;
 	FDelegateHandle SourceControlStateChangedHandle;
-};
\ No newline at end of file
+};
+
+#endif // SOURCE_CONTROL_WITH_SLATE
\ No newline at end of file
