* Implemented basic Snapshot List
* Added Snapshot creator avatar to Snapshots list
* Extended information displayed in snapshot state column in Snapshot list

diff --git a/Engine/Source/Developer/SourceControl/Private/RevisionControlStyle/RevisionControlStyle.cpp b/Engine/Source/Developer/SourceControl/Private/RevisionControlStyle/RevisionControlStyle.cpp
--- a/Engine/Source/Developer/SourceControl/Private/RevisionControlStyle/RevisionControlStyle.cpp
+++ b/Engine/Source/Developer/SourceControl/Private/RevisionControlStyle/RevisionControlStyle.cpp
@@ -95,6 +95,7 @@ FDefaultRevisionControlStyle::FDefaultRevisionControlStyle() : FSlateStyleSet(St
 	Set("RevisionControl.Actions.Merge", new IMAGE_BRUSH_SVG("Starship/Common/Merge", CoreStyleConstants::Icon16x16));
 	Set("RevisionControl.Actions.Refresh", new IMAGE_BRUSH_SVG("Starship/Common/Update", CoreStyleConstants::Icon16x16));
 	Set("RevisionControl.Actions.ChangeSettings", new IMAGE_BRUSH_SVG("Starship/SourceControl/icon_SCC_Change_Source_Control_Settings", CoreStyleConstants::Icon16x16));
+	Set("RevisionControl.Actions.Rewind", new IMAGE_BRUSH_SVG("Starship/SourceControl/RC_Rewind", CoreStyleConstants::Icon16x16));
 	Set("RevisionControl.Locked", new IMAGE_BRUSH_SVG("Starship/Common/lock", CoreStyleConstants::Icon16x16));
 
 	// Icons representing the various revision control states
@@ -131,6 +132,14 @@ FDefaultRevisionControlStyle::FDefaultRevisionControlStyle() : FSlateStyleSet(St
 	Set("RevisionControl.StatusBar.Conflicted", new IMAGE_BRUSH_SVG("Starship/SourceControl/RC_Conflicted", CoreStyleConstants::Icon16x16, FStyleColors::Warning));
 
 	Set("RevisionControl.ConflictResolution.OpenExternal", new IMAGE_BRUSH_SVG("Starship/SourceControl/RC_ConflictResolution_OpenExternal", CoreStyleConstants::Icon16x16));
+
+	// Revision Control States
+	Set("RevisionControl.VerticalLine", new IMAGE_BRUSH_SVG("Starship/SourceControl/RC_VerticalLine", CoreStyleConstants::Icon26x26));
+	Set("RevisionControl.VerticalLineDashed", new IMAGE_BRUSH_SVG("Starship/SourceControl/RC_VerticalLineDashed", CoreStyleConstants::Icon26x26));
+	Set("RevisionControl.CheckCircleLine", new IMAGE_BRUSH_SVG("Starship/SourceControl/RC_CheckCircleLine", CoreStyleConstants::Icon26x26));
+	Set("RevisionControl.LineCircle", new IMAGE_BRUSH_SVG("Starship/SourceControl/RC_LineCircle", CoreStyleConstants::Icon26x26));
+	Set("RevisionControl.CheckInAvailable", new IMAGE_BRUSH_SVG("Starship/SourceControl/RC_CheckInAvailable", CoreStyleConstants::Icon26x26));
+	Set("RevisionControl.Rewound", new IMAGE_BRUSH_SVG("Starship/SourceControl/RC_Rewound", CoreStyleConstants::Icon26x26));
 }
 
 FDefaultRevisionControlStyle::~FDefaultRevisionControlStyle()
diff --git a/Engine/Source/Runtime/SlateCore/Public/Styling/CoreStyle.h b/Engine/Source/Runtime/SlateCore/Public/Styling/CoreStyle.h
--- a/Engine/Source/Runtime/SlateCore/Public/Styling/CoreStyle.h
+++ b/Engine/Source/Runtime/SlateCore/Public/Styling/CoreStyle.h
@@ -96,6 +96,7 @@ namespace CoreStyleConstants
 	inline const UE::Slate::FDeprecateVector2DResult Icon22x22(22.0f, 22.0f);
 	inline const UE::Slate::FDeprecateVector2DResult Icon24x24(24.0f, 24.0f);
 	inline const UE::Slate::FDeprecateVector2DResult Icon25x25(25.0f, 25.0f);
+	inline const UE::Slate::FDeprecateVector2DResult Icon26x26(26.0f, 26.0f);
 	inline const UE::Slate::FDeprecateVector2DResult Icon32x32(32.0f, 32.0f);
 	inline const UE::Slate::FDeprecateVector2DResult Icon40x40(40.0f, 40.0f);
 	inline const UE::Slate::FDeprecateVector2DResult Icon64x64(64.0f, 64.0f);
