Skein / ConflictUE: Enable conflict resolution workflow and remove the obsolete retrieval of asset list pre-pass.

If a conflict situation ever ends up blocking an internal developer (due to some unforeseen issue), the CLI command 'asset resolve conclude' can always get them out of it.

#rnx

diff --git a/Engine/Source/Developer/SourceControl/Private/SSourceControlControls.cpp b/Engine/Source/Developer/SourceControl/Private/SSourceControlControls.cpp
--- a/Engine/Source/Developer/SourceControl/Private/SSourceControlControls.cpp
+++ b/Engine/Source/Developer/SourceControl/Private/SSourceControlControls.cpp
@@ -175,14 +175,7 @@ bool SSourceControlControls::AreConflictsRemaining() const
 
 	if (!bExiting)
 	{
-		// Is ConflictResolution enabled via Unreal Revision Control?
-		if (IConsoleVariable* CVar = IConsoleManager::Get().FindConsoleVariable(TEXT("UnrealRevisionControl.EnableConflictResolution")))
-		{
-			if (CVar->GetBool())
-			{
-				return bConflictsRemaining;
-			}
-		}
+		return bConflictsRemaining;
 	}
 
 	return false;
