SourceControl: Fixed Content Browser not updating after reverting a Content Browser asset move to a subfolder along with any additional action that will trigger a world reload.

Technical:
- Moving an asset to a subfolder and choosing REVERT ALL works! The asset will show up again in Content Browser.
- Moving an asset to a subfolder, moving an actor (which triggers a world reload) and choosing REVERT ALL does NOT work! The asset will not show up again in Content Browser until the project is closed and re-opened. The reason it doesn't show up is because the AssetRegistry still thinks there is a REDIRECTOR asset in that location and the Content Browser doesn't show those by default.

The cause for this is that when the two packages (for example: a redirector material package and a world package) are passed into UPackageTools::ReloadPackages, it will call CreateNewMapForEditing which will cause the redirector material package to go STALE. This is tracked by FScopedTrackFilteredPackages. This results in the redirector material package never actually being reloaded, leaving the AssetRegistry with incorrect data, causing the Content Browser to not show it.

This is fixed by doing the world reload(s) separately.

#rnx

diff --git a/Engine/Source/Developer/SourceControl/Private/SourceControlHelpers.cpp b/Engine/Source/Developer/SourceControl/Private/SourceControlHelpers.cpp
--- a/Engine/Source/Developer/SourceControl/Private/SourceControlHelpers.cpp
+++ b/Engine/Source/Developer/SourceControl/Private/SourceControlHelpers.cpp
@@ -1129,11 +1129,44 @@ bool USourceControlHelpers::ApplyOperationAndReloadPackages(const TArray<FString
 	// Hot-reload the new packages...
 	if (LoadedPackages.Num() > 0)
 	{
-		FText OutReloadErrorMsg;
-		UPackageTools::ReloadPackages(LoadedPackages, OutReloadErrorMsg, bInteractive ? EReloadPackagesInteractionMode::Interactive : EReloadPackagesInteractionMode::AssumePositive);
-		if (!OutReloadErrorMsg.IsEmpty())
+		// Split into world and non-world packages and reload them separately.
+		TArray<UPackage*> LoadedWorldPackages;
+		TArray<UPackage*> LoadedNonWorldPackages;
+		LoadedWorldPackages.Reserve(LoadedPackages.Num());
+		LoadedNonWorldPackages.Reserve(LoadedPackages.Num());
+
+		for (UPackage* Package : LoadedPackages)
+		{
+			if (UWorld::FindWorldInPackage(Package))
+			{
+				LoadedWorldPackages.Add(Package);
+			}
+			else
+			{
+				LoadedNonWorldPackages.Add(Package);
+			}
+		}
+
+		// Reload non world package(s).
+		if (LoadedNonWorldPackages.Num() > 0)
+		{
+			FText OutReloadErrorMsg;
+			UPackageTools::ReloadPackages(LoadedNonWorldPackages, OutReloadErrorMsg, bInteractive ? EReloadPackagesInteractionMode::Interactive : EReloadPackagesInteractionMode::AssumePositive);
+			if (!OutReloadErrorMsg.IsEmpty())
+			{
+				UE_LOG(LogSourceControl, Warning, TEXT("%s"), *OutReloadErrorMsg.ToString());
+			}
+		}
+
+		// Reload world package(s).
+		if (LoadedWorldPackages.Num() > 0)
 		{
-			UE_LOG(LogSourceControl, Warning, TEXT("%s"), *OutReloadErrorMsg.ToString());
+			FText OutReloadErrorMsg;
+			UPackageTools::ReloadPackages(LoadedWorldPackages, OutReloadErrorMsg, bInteractive ? EReloadPackagesInteractionMode::Interactive : EReloadPackagesInteractionMode::AssumePositive);
+			if (!OutReloadErrorMsg.IsEmpty())
+			{
+				UE_LOG(LogSourceControl, Warning, TEXT("%s"), *OutReloadErrorMsg.ToString());
+			}
 		}
 	}
 
