Added warning highlight to Asset List [No.] column when user has rewound to an older revision

diff --git a/Engine/Source/Developer/SourceControl/Private/RevisionControlStyle/RevisionControlStyle.cpp b/Engine/Source/Developer/SourceControl/Private/RevisionControlStyle/RevisionControlStyle.cpp
--- a/Engine/Source/Developer/SourceControl/Private/RevisionControlStyle/RevisionControlStyle.cpp
+++ b/Engine/Source/Developer/SourceControl/Private/RevisionControlStyle/RevisionControlStyle.cpp
@@ -141,6 +141,7 @@ FDefaultRevisionControlStyle::FDefaultRevisionControlStyle() : FSlateStyleSet(St
 	Set("RevisionControl.LineCircle", new IMAGE_BRUSH_SVG("Starship/SourceControl/RC_LineCircle", CoreStyleConstants::Icon26x26));
 	Set("RevisionControl.CheckInAvailable", new IMAGE_BRUSH_SVG("Starship/SourceControl/RC_CheckInAvailable", CoreStyleConstants::Icon26x26));
 	Set("RevisionControl.Rewound", new IMAGE_BRUSH_SVG("Starship/SourceControl/RC_Rewound", CoreStyleConstants::Icon26x26));
+	Set("RevisionControl.CheckInAvailableRewound", new IMAGE_BRUSH_SVG("Starship/SourceControl/RC_CheckInAvailableRewound", CoreStyleConstants::Icon26x26));
 }
 
 FDefaultRevisionControlStyle::~FDefaultRevisionControlStyle()
