Skein / ConflictUE: Fixed 'check in changes' no longer redirecting to the 'conflict resolution' window when there are conflicts remaining. This was introduced in #29863963 when SSourceControlControls was created.

#rnx

diff --git a/Engine/Source/Developer/SourceControl/Private/SSourceControlControls.cpp b/Engine/Source/Developer/SourceControl/Private/SSourceControlControls.cpp
--- a/Engine/Source/Developer/SourceControl/Private/SSourceControlControls.cpp
+++ b/Engine/Source/Developer/SourceControl/Private/SSourceControlControls.cpp
@@ -60,7 +60,7 @@ void SSourceControlControls::Construct(const FArguments& InArgs)
 					.Text(this, &SSourceControlControls::GetSourceControlCheckInStatusText)
 				]
 			]
-			.OnClicked(InArgs._OnClickedCheckInChanges)
+			.OnClicked(this, &SSourceControlControls::OnSourceControlCheckInChangesClicked)
 		]
 		+SHorizontalBox::Slot() // Check In Kebab Combo button
 		.VAlign(VAlign_Center)
