Skein / ConflictUE: Assorted conflict resolution UI updates. See JIRA for checkmarked list.

#rnx

diff --git a/Engine/Source/Developer/SourceControl/Private/RevisionControlStyle/RevisionControlStyle.cpp b/Engine/Source/Developer/SourceControl/Private/RevisionControlStyle/RevisionControlStyle.cpp
--- a/Engine/Source/Developer/SourceControl/Private/RevisionControlStyle/RevisionControlStyle.cpp
+++ b/Engine/Source/Developer/SourceControl/Private/RevisionControlStyle/RevisionControlStyle.cpp
@@ -120,7 +120,7 @@ FDefaultRevisionControlStyle::FDefaultRevisionControlStyle() : FSlateStyleSet(St
 	
 	Set("RevisionControl.Branched", new IMAGE_BRUSH_SVG("Starship/SourceControl/SCC_Branched", CoreStyleConstants::Icon16x16, BranchedColor));
 
-	Set("RevisionControl.Conflicted", new IMAGE_BRUSH_SVG("Starship/SourceControl/RC_Conflicted", CoreStyleConstants::Icon16x16, FStyleColors::Error));
+	Set("RevisionControl.Conflicted", new IMAGE_BRUSH_SVG("Starship/SourceControl/RC_Conflicted", CoreStyleConstants::Icon16x16));
 	
 	// Misc Icons
 	Set("RevisionControl.ChangelistsTab", new IMAGE_BRUSH_SVG("Starship/Common/check-circle", CoreStyleConstants::Icon16x16));
@@ -129,9 +129,10 @@ FDefaultRevisionControlStyle::FDefaultRevisionControlStyle() : FSlateStyleSet(St
 	Set("RevisionControl.StatusBar.NotAtLatestRevision", new IMAGE_BRUSH_SVG("Starship/SourceControl/RC_StatusRemoteDownload", CoreStyleConstants::Icon16x16, FStyleColors::Warning));
 	Set("RevisionControl.StatusBar.NoLocalChanges", new IMAGE_BRUSH_SVG("Starship/SourceControl/RC_StatusLocalUpToDate", CoreStyleConstants::Icon16x16));
 	Set("RevisionControl.StatusBar.HasLocalChanges", new IMAGE_BRUSH_SVG("Starship/SourceControl/RC_StatusLocalUpload", CoreStyleConstants::Icon16x16, FStyleColors::AccentBlue));
-	Set("RevisionControl.StatusBar.Conflicted", new IMAGE_BRUSH_SVG("Starship/SourceControl/RC_Conflicted", CoreStyleConstants::Icon16x16, FStyleColors::Warning));
+	Set("RevisionControl.StatusBar.Conflicted", new IMAGE_BRUSH_SVG("Starship/SourceControl/RC_Conflicted", CoreStyleConstants::Icon16x16));
 
 	Set("RevisionControl.ConflictResolution.OpenExternal", new IMAGE_BRUSH_SVG("Starship/SourceControl/RC_ConflictResolution_OpenExternal", CoreStyleConstants::Icon16x16));
+	Set("RevisionControl.ConflictResolution.Clear", new IMAGE_BRUSH_SVG("Starship/SourceControl/RC_ConflictResolution_Clear", CoreStyleConstants::Icon16x16));
 
 	// Revision Control States
 	Set("RevisionControl.VerticalLine", new IMAGE_BRUSH_SVG("Starship/SourceControl/RC_VerticalLine", CoreStyleConstants::Icon26x26));
