Skein / Snapshot History: Assorted snapshot history UI updates. See JIRA for checkmarked list.

#rnx

diff --git a/Engine/Source/Developer/SourceControl/Private/SSourceControlControls.cpp b/Engine/Source/Developer/SourceControl/Private/SSourceControlControls.cpp
--- a/Engine/Source/Developer/SourceControl/Private/SSourceControlControls.cpp
+++ b/Engine/Source/Developer/SourceControl/Private/SSourceControlControls.cpp
@@ -27,6 +27,9 @@ void SSourceControlControls::Construct(const FArguments& InArgs)
 	IsSyncLatestEnabled = InArgs._IsEnabledSyncLatest;
 	IsCheckInChangesEnabled = InArgs._IsEnabledCheckInChanges;
 
+	IsSyncLatestSeparatorEnabled = InArgs._IsEnabledSyncLatestSeparator;
+	IsCheckInChangesSeparatorEnabled = InArgs._IsEnabledCheckInChangesSeparator;
+
 	ChildSlot
 	[
 		SNew(SHorizontalBox)
@@ -78,7 +81,7 @@ void SSourceControlControls::Construct(const FArguments& InArgs)
 		.AutoWidth()
 		[
 			SNew(SSeparator)
-			.Visibility(this, &SSourceControlControls::GetSourceControlCheckInStatusVisibility)
+			.Visibility(this, &SSourceControlControls::GetSourceControlCheckInSeparatorVisibility)
 			.Thickness(1.0)
 			.Orientation(EOrientation::Orient_Vertical)
 		]
@@ -118,7 +121,7 @@ void SSourceControlControls::Construct(const FArguments& InArgs)
 		.AutoWidth()
 		[
 			SNew(SSeparator)
-			.Visibility(this, &SSourceControlControls::GetSourceControlCheckInStatusVisibility)
+			.Visibility(this, &SSourceControlControls::GetSourceControlSyncSeparatorVisibility)
 			.Thickness(1.0)
 			.Orientation(EOrientation::Orient_Vertical)
 		]
@@ -234,6 +237,17 @@ EVisibility SSourceControlControls::GetSourceControlSyncStatusVisibility() const
 	return EVisibility::Collapsed;
 }
 
+EVisibility SSourceControlControls::GetSourceControlSyncSeparatorVisibility() const
+{
+	EVisibility StatusVisibility = GetSourceControlSyncStatusVisibility();
+	if (StatusVisibility != EVisibility::Visible)
+	{
+		return StatusVisibility;
+	}
+
+	return IsSyncLatestSeparatorEnabled.Get(true) ? EVisibility::Visible : EVisibility::Collapsed;
+}
+
 FText SSourceControlControls::GetSourceControlSyncStatusText() const
 {
 	if (HasSourceControlChangesToSync())
@@ -346,6 +360,17 @@ EVisibility SSourceControlControls::GetSourceControlCheckInStatusVisibility() co
 	return EVisibility::Collapsed;
 }
 
+EVisibility SSourceControlControls::GetSourceControlCheckInSeparatorVisibility() const
+{
+	EVisibility StatusVisibility = GetSourceControlCheckInStatusVisibility();
+	if (StatusVisibility != EVisibility::Visible)
+	{
+		return StatusVisibility;
+	}
+
+	return IsCheckInChangesSeparatorEnabled.Get(true) ? EVisibility::Visible : EVisibility::Collapsed;
+}
+
 FText SSourceControlControls::GetSourceControlCheckInStatusText() const
 {
 	if (HasSourceControlChangesToCheckIn())
diff --git a/Engine/Source/Developer/SourceControl/Public/SSourceControlControls.h b/Engine/Source/Developer/SourceControl/Public/SSourceControlControls.h
--- a/Engine/Source/Developer/SourceControl/Public/SSourceControlControls.h
+++ b/Engine/Source/Developer/SourceControl/Public/SSourceControlControls.h
@@ -19,6 +19,8 @@ public:
 
 		SLATE_ATTRIBUTE(bool, IsEnabledSyncLatest)
 		SLATE_ATTRIBUTE(bool, IsEnabledCheckInChanges)
+		SLATE_ATTRIBUTE(bool, IsEnabledSyncLatestSeparator)
+		SLATE_ATTRIBUTE(bool, IsEnabledCheckInChangesSeparator)
 		SLATE_EVENT(FOnClicked, OnClickedSyncLatest)
 		SLATE_EVENT(FOnClicked, OnClickedCheckInChanges)
 		SLATE_EVENT(FOnGetContent, OnGenerateKebabMenu)
@@ -35,6 +37,7 @@ private:
 	bool IsSourceControlSyncEnabled() const;
 	bool HasSourceControlChangesToSync() const;
 	EVisibility GetSourceControlSyncStatusVisibility() const;
+	EVisibility GetSourceControlSyncSeparatorVisibility() const;
 	FText GetSourceControlSyncStatusText() const;
 	FText GetSourceControlSyncStatusTooltipText() const;
 	const FSlateBrush* GetSourceControlSyncStatusIcon() const;
@@ -45,6 +48,7 @@ private:
 	bool IsSourceControlCheckInEnabled() const;
 	bool HasSourceControlChangesToCheckIn() const;
 	EVisibility GetSourceControlCheckInStatusVisibility() const;
+	EVisibility GetSourceControlCheckInSeparatorVisibility() const;
 	FText GetSourceControlCheckInStatusText() const;
 	FText GetSourceControlCheckInStatusTooltipText() const;
 	const FSlateBrush* GetSourceControlCheckInStatusIcon() const;
@@ -65,6 +69,9 @@ private:
 	TAttribute<bool> IsSyncLatestEnabled;
 	TAttribute<bool> IsCheckInChangesEnabled;
 
+	TAttribute<bool> IsSyncLatestSeparatorEnabled;
+	TAttribute<bool> IsCheckInChangesSeparatorEnabled;
+
 	/** Is there a conflict remaining? */
 	bool bConflictsRemaining;
 
