Populate Assets List in expanded Snapshot view
Created Asset Class 'filter pill' widget
Extended Assets List to display a colour pill representing the asset class

diff --git a/Engine/Source/Developer/SourceControl/Private/RevisionControlStyle/RevisionControlStyle.cpp b/Engine/Source/Developer/SourceControl/Private/RevisionControlStyle/RevisionControlStyle.cpp
--- a/Engine/Source/Developer/SourceControl/Private/RevisionControlStyle/RevisionControlStyle.cpp
+++ b/Engine/Source/Developer/SourceControl/Private/RevisionControlStyle/RevisionControlStyle.cpp
@@ -143,6 +143,7 @@ FDefaultRevisionControlStyle::FDefaultRevisionControlStyle() : FSlateStyleSet(St
 	Set("RevisionControl.CheckInAvailable", new IMAGE_BRUSH_SVG("Starship/SourceControl/RC_CheckInAvailable", CoreStyleConstants::Icon26x26));
 	Set("RevisionControl.Rewound", new IMAGE_BRUSH_SVG("Starship/SourceControl/RC_Rewound", CoreStyleConstants::Icon26x26));
 	Set("RevisionControl.CheckInAvailableRewound", new IMAGE_BRUSH_SVG("Starship/SourceControl/RC_CheckInAvailableRewound", CoreStyleConstants::Icon26x26));
+	Set("RevisionControl.ConflictedState", new IMAGE_BRUSH_SVG("Starship/SourceControl/RC_ConflictedState", CoreStyleConstants::Icon26x26));
 }
 
 FDefaultRevisionControlStyle::~FDefaultRevisionControlStyle()
diff --git a/Engine/Source/Developer/SourceControl/Private/SSourceControlControls.cpp b/Engine/Source/Developer/SourceControl/Private/SSourceControlControls.cpp
--- a/Engine/Source/Developer/SourceControl/Private/SSourceControlControls.cpp
+++ b/Engine/Source/Developer/SourceControl/Private/SSourceControlControls.cpp
@@ -14,6 +14,10 @@
 
 #define LOCTEXT_NAMESPACE "SSkeinSourceControlWidgets"
 
+bool SSourceControlControls::bStaticDisableCheckInChangesOverride = false;
+bool SSourceControlControls::bStaticDisableSyncLatestOverride = false;
+FOnClicked SSourceControlControls::OnSyncLatestClickedStaticOverride;
+
 /**
  * Construct this widget
  *
@@ -172,7 +176,8 @@ void SSourceControlControls::OnSourceControlStateChanged()
 		}
 	);
 
-	bConflictsRemaining = (Conflicts.Num() > 0);
+	NumConflictsRemaining = Conflicts.Num();
+	bConflictsRemaining = (NumConflictsRemaining > 0);
 }
 
 bool SSourceControlControls::AreConflictsRemaining() const
@@ -187,6 +192,11 @@ bool SSourceControlControls::AreConflictsRemaining() const
 	return false;
 }
 
+int32 SSourceControlControls::GetNumConflictsRemaining() const
+{
+	return NumConflictsRemaining;
+}
+
 /** Sync Status */
 
 bool SSourceControlControls::IsAtLatestRevision() const
@@ -205,7 +215,7 @@ bool SSourceControlControls::IsSourceControlSyncEnabled() const
 		return false;
 	}
 
-	return IsSyncLatestEnabled.Get(true);
+	return IsSyncLatestEnabled.Get(true) && !bStaticDisableSyncLatestOverride;
 }
 
 bool SSourceControlControls::HasSourceControlChangesToSync() const
@@ -298,7 +308,11 @@ FReply SSourceControlControls::OnSourceControlSyncClicked() const
 	}
 	else if (HasSourceControlChangesToSync())
 	{
-		if (OnSyncLatestClicked.IsBound())
+		if (OnSyncLatestClickedStaticOverride.IsBound())
+		{
+			OnSyncLatestClickedStaticOverride.Execute();
+		}
+		else if (OnSyncLatestClicked.IsBound())
 		{
 			OnSyncLatestClicked.Execute();
 		}
@@ -328,7 +342,7 @@ bool SSourceControlControls::IsSourceControlCheckInEnabled() const
 		return false;
 	}
 
-	return IsCheckInChangesEnabled.Get(true);
+	return IsCheckInChangesEnabled.Get(true) && !bStaticDisableCheckInChangesOverride;
 }
 
 bool SSourceControlControls::HasSourceControlChangesToCheckIn() const
@@ -380,7 +394,6 @@ FText SSourceControlControls::GetSourceControlCheckInStatusText() const
 	return LOCTEXT("CheckInButtonNoChangesText", "No Changes");
 }
 
-
 FText SSourceControlControls::GetSourceControlCheckInStatusTooltipText() const
 {
 	if (AreConflictsRemaining())
@@ -389,6 +402,10 @@ FText SSourceControlControls::GetSourceControlCheckInStatusTooltipText() const
 	}
 	if (HasSourceControlChangesToCheckIn())
 	{
+		if (bStaticDisableCheckInChangesOverride)
+		{
+			return LOCTEXT("CheckInButtonChangesTooltipTextDisabled", "Check-in disabled while in a rewound state.");
+		}
 		return FText::Format(LOCTEXT("CheckInButtonChangesTooltipText", "Check-in {0} change(s) to this project"), GetNumLocalChanges());
 	}
 	return LOCTEXT("CheckInButtonNoChangesTooltipText", "No Changes to check in for this project");
diff --git a/Engine/Source/Developer/SourceControl/Public/SSourceControlControls.h b/Engine/Source/Developer/SourceControl/Public/SSourceControlControls.h
--- a/Engine/Source/Developer/SourceControl/Public/SSourceControlControls.h
+++ b/Engine/Source/Developer/SourceControl/Public/SSourceControlControls.h
@@ -56,8 +56,15 @@ private:
 
 	/** Conflicts */
 	void CheckSourceControlStatus();
+public:
 	bool AreConflictsRemaining() const;
+	int32 GetNumConflictsRemaining() const;
+	
+	static void SetCheckInChangesDisabledOverride(bool InEnabled) { bStaticDisableCheckInChangesOverride = InEnabled; }
+	static void SetSyncLatestDisabledOverride(bool InEnabled) { bStaticDisableSyncLatestOverride = InEnabled; }
+	static FOnClicked& GetOnSyncLatestClickedStaticOverride() { return OnSyncLatestClickedStaticOverride; }
 
+private:
 	void OnSourceControlProviderChanged(ISourceControlProvider& OldProvider, ISourceControlProvider& NewProvider);
 	void OnSourceControlStateChanged();
 	
@@ -74,9 +81,14 @@ private:
 
 	/** Is there a conflict remaining? */
 	bool bConflictsRemaining;
+	int32 NumConflictsRemaining;
 
 	FDelegateHandle SourceControlProviderChangedHandle;
 	FDelegateHandle SourceControlStateChangedHandle;
+
+	static bool bStaticDisableCheckInChangesOverride;
+	static bool bStaticDisableSyncLatestOverride;
+	static FOnClicked OnSyncLatestClickedStaticOverride;
 };
 
 #endif // SOURCE_CONTROL_WITH_SLATE
\ No newline at end of file
