SkeinUE: Changed `conflicted` icon used by Skein to be white (in the file) with a warning color applied instead of being a colored icon in the file.

#rnx

diff --git a/Engine/Source/Developer/SourceControl/Private/RevisionControlStyle/RevisionControlStyle.cpp b/Engine/Source/Developer/SourceControl/Private/RevisionControlStyle/RevisionControlStyle.cpp
--- a/Engine/Source/Developer/SourceControl/Private/RevisionControlStyle/RevisionControlStyle.cpp
+++ b/Engine/Source/Developer/SourceControl/Private/RevisionControlStyle/RevisionControlStyle.cpp
@@ -129,7 +129,7 @@ FDefaultRevisionControlStyle::FDefaultRevisionControlStyle() : FSlateStyleSet(St
 	Set("RevisionControl.StatusBar.NotAtLatestRevision", new IMAGE_BRUSH_SVG("Starship/SourceControl/RC_StatusRemoteDownload", CoreStyleConstants::Icon16x16, FStyleColors::Warning));
 	Set("RevisionControl.StatusBar.NoLocalChanges", new IMAGE_BRUSH_SVG("Starship/SourceControl/RC_StatusLocalUpToDate", CoreStyleConstants::Icon16x16));
 	Set("RevisionControl.StatusBar.HasLocalChanges", new IMAGE_BRUSH_SVG("Starship/SourceControl/RC_StatusLocalUpload", CoreStyleConstants::Icon16x16, FStyleColors::AccentBlue));
-	Set("RevisionControl.StatusBar.Conflicted", new IMAGE_BRUSH_SVG("Starship/SourceControl/RC_Conflicted", CoreStyleConstants::Icon16x16));
+	Set("RevisionControl.StatusBar.Conflicted", new IMAGE_BRUSH_SVG("Starship/SourceControl/RC_Conflicted", CoreStyleConstants::Icon16x16, FStyleColors::Warning));
 
 	Set("RevisionControl.ConflictResolution.OpenExternal", new IMAGE_BRUSH_SVG("Starship/SourceControl/RC_ConflictResolution_OpenExternal", CoreStyleConstants::Icon16x16));
 	Set("RevisionControl.ConflictResolution.Clear", new IMAGE_BRUSH_SVG("Starship/SourceControl/RC_ConflictResolution_Clear", CoreStyleConstants::Icon16x16));
