Fix for "icons for asset change states aren't displaying" in Snapshot History tab

#localization none
#tests ran local build and preflight

diff --git a/Engine/Source/Developer/SourceControl/Private/RevisionControlStyle/RevisionControlStyle.cpp b/Engine/Source/Developer/SourceControl/Private/RevisionControlStyle/RevisionControlStyle.cpp
--- a/Engine/Source/Developer/SourceControl/Private/RevisionControlStyle/RevisionControlStyle.cpp
+++ b/Engine/Source/Developer/SourceControl/Private/RevisionControlStyle/RevisionControlStyle.cpp
@@ -144,6 +144,11 @@ FDefaultRevisionControlStyle::FDefaultRevisionControlStyle() : FSlateStyleSet(St
 	Set("RevisionControl.Rewound", new IMAGE_BRUSH_SVG("Starship/SourceControl/RC_Rewound", CoreStyleConstants::Icon26x26));
 	Set("RevisionControl.CheckInAvailableRewound", new IMAGE_BRUSH_SVG("Starship/SourceControl/RC_CheckInAvailableRewound", CoreStyleConstants::Icon26x26));
 	Set("RevisionControl.ConflictedState", new IMAGE_BRUSH_SVG("Starship/SourceControl/RC_ConflictedState", CoreStyleConstants::Icon26x26));
+
+	Set("RevisionControl.File", new IMAGE_BRUSH_SVG("Starship/SourceControl/RC_File", CoreStyleConstants::Icon6x8));
+	Set("RevisionControl.State.Added", new IMAGE_BRUSH_SVG("Starship/SourceControl/RC_Added", CoreStyleConstants::Icon12x12));
+	Set("RevisionControl.State.Modified", new IMAGE_BRUSH_SVG("Starship/SourceControl/RC_Modified", CoreStyleConstants::Icon12x12));
+	Set("RevisionControl.State.Removed", new IMAGE_BRUSH_SVG("Starship/SourceControl/RC_Removed", CoreStyleConstants::Icon12x12));
 }
 
 FDefaultRevisionControlStyle::~FDefaultRevisionControlStyle()
diff --git a/Engine/Source/Runtime/SlateCore/Public/Styling/CoreStyle.h b/Engine/Source/Runtime/SlateCore/Public/Styling/CoreStyle.h
--- a/Engine/Source/Runtime/SlateCore/Public/Styling/CoreStyle.h
+++ b/Engine/Source/Runtime/SlateCore/Public/Styling/CoreStyle.h
@@ -83,6 +83,7 @@ namespace CoreStyleConstants
 	// Note, these sizes are in Slate Units.
 	// Slate Units do NOT have to map to pixels.
 	inline const UE::Slate::FDeprecateVector2DResult Icon5x16(5.0f, 16.0f);
+	inline const UE::Slate::FDeprecateVector2DResult Icon6x8(6.0f, 8.0f);
 	inline const UE::Slate::FDeprecateVector2DResult Icon8x4(8.0f, 4.0f);
 	inline const UE::Slate::FDeprecateVector2DResult Icon16x4(16.0f, 4.0f);
 	inline const UE::Slate::FDeprecateVector2DResult Icon8x8(8.0f, 8.0f);
