SourceControl: Fixed an error being reported for failing to convert to a package name for non-package files (for example: *.uplugin). This was introduced in #24242330.

#rnx

diff --git a/Engine/Source/Developer/SourceControl/Private/SourceControlHelpers.cpp b/Engine/Source/Developer/SourceControl/Private/SourceControlHelpers.cpp
--- a/Engine/Source/Developer/SourceControl/Private/SourceControlHelpers.cpp
+++ b/Engine/Source/Developer/SourceControl/Private/SourceControlHelpers.cpp
@@ -1261,7 +1261,7 @@ bool USourceControlHelpers::ListRevertablePackages(TArray<FString>& OutRevertabl
 		return false;
 	}
 
-	// Get a list of all the revertable packages	
+	// Get a list of all the revertable packages
 	TMap<FString, FSourceControlStatePtr> PackageStates;
 	FEditorFileUtils::FindAllSubmittablePackageFiles(PackageStates, true);
 
@@ -1275,19 +1275,23 @@ bool USourceControlHelpers::ListRevertablePackages(TArray<FString>& OutRevertabl
 		{
 			const FString& Filename = State->GetFilename();
 
-			FString PackageName;
-			FString FailureReason;
-			if (!FPackageName::TryConvertFilenameToLongPackageName(Filename, PackageName, &FailureReason))
+			if (FPackageName::IsPackageFilename(Filename))
 			{
-				UE_LOG(LogSourceControl, Warning, TEXT("%s"), *FailureReason);
-				return false;
-			}
+				FString PackageName;
+				FString FailureReason;
+				if (!FPackageName::TryConvertFilenameToLongPackageName(Filename, PackageName, &FailureReason))
+				{
+					UE_LOG(LogSourceControl, Warning, TEXT("%s"), *FailureReason);
+					return false;
+				}
 
-			if (State->IsDeleted() && !PackageNames.Contains(PackageName))
-			{
-				FileNamesToPackageNames.Add(Filename, PackageName);
-				return true;
+				if (State->IsDeleted() && !PackageNames.Contains(PackageName))
+				{
+					FileNamesToPackageNames.Add(Filename, PackageName);
+					return true;
+				}
 			}
+
 			return false;
 		}
 	);
