SkeinUE: Fixed conflicted icon incorrectly being white in some areas since #30499647.

#rnx

diff --git a/Engine/Source/Developer/SourceControl/Private/RevisionControlStyle/RevisionControlStyle.cpp b/Engine/Source/Developer/SourceControl/Private/RevisionControlStyle/RevisionControlStyle.cpp
--- a/Engine/Source/Developer/SourceControl/Private/RevisionControlStyle/RevisionControlStyle.cpp
+++ b/Engine/Source/Developer/SourceControl/Private/RevisionControlStyle/RevisionControlStyle.cpp
@@ -120,7 +120,7 @@ FDefaultRevisionControlStyle::FDefaultRevisionControlStyle() : FSlateStyleSet(St
 	
 	Set("RevisionControl.Branched", new IMAGE_BRUSH_SVG("Starship/SourceControl/SCC_Branched", CoreStyleConstants::Icon16x16, BranchedColor));
 
-	Set("RevisionControl.Conflicted", new IMAGE_BRUSH_SVG("Starship/SourceControl/RC_Conflicted", CoreStyleConstants::Icon16x16));
+	Set("RevisionControl.Conflicted", new IMAGE_BRUSH_SVG("Starship/SourceControl/RC_Conflicted", CoreStyleConstants::Icon16x16, FStyleColors::Warning));
 	
 	// Misc Icons
 	Set("RevisionControl.ChangelistsTab", new IMAGE_BRUSH_SVG("Starship/Common/check-circle", CoreStyleConstants::Icon16x16));
