* Added Restore as Latest functionality
* Created 'Restore as latest' dialog

diff --git a/Engine/Source/Developer/SourceControl/Private/RevisionControlStyle/RevisionControlStyle.cpp b/Engine/Source/Developer/SourceControl/Private/RevisionControlStyle/RevisionControlStyle.cpp
--- a/Engine/Source/Developer/SourceControl/Private/RevisionControlStyle/RevisionControlStyle.cpp
+++ b/Engine/Source/Developer/SourceControl/Private/RevisionControlStyle/RevisionControlStyle.cpp
@@ -130,6 +130,10 @@ FDefaultRevisionControlStyle::FDefaultRevisionControlStyle() : FSlateStyleSet(St
 	Set("RevisionControl.StatusBar.NoLocalChanges", new IMAGE_BRUSH_SVG("Starship/SourceControl/RC_StatusLocalUpToDate", CoreStyleConstants::Icon16x16));
 	Set("RevisionControl.StatusBar.HasLocalChanges", new IMAGE_BRUSH_SVG("Starship/SourceControl/RC_StatusLocalUpload", CoreStyleConstants::Icon16x16, FStyleColors::AccentBlue));
 	Set("RevisionControl.StatusBar.Conflicted", new IMAGE_BRUSH_SVG("Starship/SourceControl/RC_Conflicted", CoreStyleConstants::Icon16x16, FStyleColors::Warning));
+	Set("RevisionControl.StatusBar.Promote", new IMAGE_BRUSH_SVG("Starship/SourceControl/RC_Promote", CoreStyleConstants::Icon16x16));
+	
+	Set("RevisionControl.Promote.Large", new IMAGE_BRUSH_SVG("Starship/SourceControl/RC_Promote_Large", CoreStyleConstants::Icon32x32));
+
 
 	Set("RevisionControl.ConflictResolution.OpenExternal", new IMAGE_BRUSH_SVG("Starship/SourceControl/RC_ConflictResolution_OpenExternal", CoreStyleConstants::Icon16x16));
 	Set("RevisionControl.ConflictResolution.Clear", new IMAGE_BRUSH_SVG("Starship/SourceControl/RC_ConflictResolution_Clear", CoreStyleConstants::Icon16x16));
diff --git a/Engine/Source/Developer/SourceControl/Private/SSourceControlControls.cpp b/Engine/Source/Developer/SourceControl/Private/SSourceControlControls.cpp
--- a/Engine/Source/Developer/SourceControl/Private/SSourceControlControls.cpp
+++ b/Engine/Source/Developer/SourceControl/Private/SSourceControlControls.cpp
@@ -14,9 +14,10 @@
 
 #define LOCTEXT_NAMESPACE "SSkeinSourceControlWidgets"
 
-bool SSourceControlControls::bStaticDisableCheckInChangesOverride = false;
+bool SSourceControlControls::bRewoundMode = false;
 bool SSourceControlControls::bStaticDisableSyncLatestOverride = false;
 FOnClicked SSourceControlControls::OnSyncLatestClickedStaticOverride;
+FOnClicked SSourceControlControls::OnRestoreAsLatestClickedStaticOverride;
 
 /**
  * Construct this widget
@@ -27,6 +28,7 @@ void SSourceControlControls::Construct(const FArguments& InArgs)
 {
 	OnSyncLatestClicked = InArgs._OnClickedSyncLatest;
 	OnCheckInChangesClicked = InArgs._OnClickedCheckInChanges;
+	OnRestoreAsLatestClicked = InArgs._OnClickedRestoreAsLatest;
 
 	IsSyncLatestEnabled = InArgs._IsEnabledSyncLatest;
 	IsCheckInChangesEnabled = InArgs._IsEnabledCheckInChanges;
@@ -69,6 +71,37 @@ void SSourceControlControls::Construct(const FArguments& InArgs)
 			]
 			.OnClicked(this, &SSourceControlControls::OnSourceControlCheckInChangesClicked)
 		]
+		+ SHorizontalBox::Slot() // Restore as Latest button
+		.VAlign(VAlign_Center)
+		.Padding(FMargin(0.0f, 0.0f, 4.0f, 0.0f))
+		.AutoWidth()
+		[
+			SNew(SButton)
+			.ButtonStyle(&FAppStyle::Get().GetWidgetStyle<FButtonStyle>("StatusBar.StatusBarButton"))
+			.ToolTipText(LOCTEXT("RestoreAsLatestTooltipText", "Restore this snapshot to be the latest version of the project for all team members."))
+			.Visibility(this, &SSourceControlControls::GetSourceControlRestoreAsLatestVisibility)
+			[
+				SNew(SHorizontalBox)
+				+ SHorizontalBox::Slot()
+				.AutoWidth()
+				.VAlign(VAlign_Center)
+				.HAlign(HAlign_Center)
+				[
+					SNew(SImage)
+					.Image(FRevisionControlStyleManager::Get().GetBrush("RevisionControl.StatusBar.Promote"))
+				]
+				+ SHorizontalBox::Slot()
+				.AutoWidth()
+				.VAlign(VAlign_Center)
+				.Padding(FMargin(5, 0, 0, 0))
+				[
+					SNew(STextBlock)
+					.TextStyle(&FAppStyle::Get().GetWidgetStyle<FTextBlockStyle>("NormalText"))
+					.Text(LOCTEXT("RestoreAsLatestButtonText", "Restore as Latest"))
+				]
+			]
+			.OnClicked(this, &SSourceControlControls::OnSourceControlRestoreAsLatestClicked)
+		]
 		+SHorizontalBox::Slot() // Check In Kebab Combo button
 		.VAlign(VAlign_Center)
 		.AutoWidth()
@@ -342,7 +375,7 @@ bool SSourceControlControls::IsSourceControlCheckInEnabled() const
 		return false;
 	}
 
-	return IsCheckInChangesEnabled.Get(true) && !bStaticDisableCheckInChangesOverride;
+	return IsCheckInChangesEnabled.Get(true);
 }
 
 bool SSourceControlControls::HasSourceControlChangesToCheckIn() const
@@ -358,6 +391,11 @@ EVisibility SSourceControlControls::GetSourceControlCheckInStatusVisibility() co
 		return EVisibility::Visible;
 	}
 
+	if (bRewoundMode)
+	{
+		return EVisibility::Collapsed;
+	}
+
 	bool bDisplaySourceControlCheckInStatus = false;
 	GConfig->GetBool(TEXT("SourceControlSettings"), TEXT("DisplaySourceControlCheckInStatus"), bDisplaySourceControlCheckInStatus, GEditorIni);
 
@@ -402,10 +440,6 @@ FText SSourceControlControls::GetSourceControlCheckInStatusTooltipText() const
 	}
 	if (HasSourceControlChangesToCheckIn())
 	{
-		if (bStaticDisableCheckInChangesOverride)
-		{
-			return LOCTEXT("CheckInButtonChangesTooltipTextDisabled", "Check-in disabled while in a rewound state.");
-		}
 		return FText::Format(LOCTEXT("CheckInButtonChangesTooltipText", "Check-in {0} change(s) to this project"), GetNumLocalChanges());
 	}
 	return LOCTEXT("CheckInButtonNoChangesTooltipText", "No Changes to check in for this project");
@@ -448,6 +482,31 @@ FReply SSourceControlControls::OnSourceControlCheckInChangesClicked() const
 	return FReply::Handled();
 }
 
+/** Restore as Latest */
+
+EVisibility SSourceControlControls::GetSourceControlRestoreAsLatestVisibility() const
+{
+	if (bRewoundMode)
+	{
+		return EVisibility::Visible;
+	}
+	return EVisibility::Collapsed;
+}
+
+FReply SSourceControlControls::OnSourceControlRestoreAsLatestClicked() const
+{
+	if (OnRestoreAsLatestClickedStaticOverride.IsBound())
+	{
+		OnRestoreAsLatestClickedStaticOverride.Execute();
+	}
+	else if (OnRestoreAsLatestClicked.IsBound())
+	{
+		OnRestoreAsLatestClicked.Execute();
+	}
+
+	return FReply::Handled();
+}
+
 
 #undef LOCTEXT_NAMESPACE
 
diff --git a/Engine/Source/Developer/SourceControl/Public/SSourceControlControls.h b/Engine/Source/Developer/SourceControl/Public/SSourceControlControls.h
--- a/Engine/Source/Developer/SourceControl/Public/SSourceControlControls.h
+++ b/Engine/Source/Developer/SourceControl/Public/SSourceControlControls.h
@@ -19,10 +19,12 @@ public:
 
 		SLATE_ATTRIBUTE(bool, IsEnabledSyncLatest)
 		SLATE_ATTRIBUTE(bool, IsEnabledCheckInChanges)
+		SLATE_ATTRIBUTE(bool, IsEnabledRestoreAsLatest)
 		SLATE_ATTRIBUTE(bool, IsEnabledSyncLatestSeparator)
 		SLATE_ATTRIBUTE(bool, IsEnabledCheckInChangesSeparator)
 		SLATE_EVENT(FOnClicked, OnClickedSyncLatest)
 		SLATE_EVENT(FOnClicked, OnClickedCheckInChanges)
+		SLATE_EVENT(FOnClicked, OnClickedRestoreAsLatest)
 		SLATE_EVENT(FOnGetContent, OnGenerateKebabMenu)
 
 	SLATE_END_ARGS()
@@ -48,21 +50,27 @@ private:
 	bool IsSourceControlCheckInEnabled() const;
 	bool HasSourceControlChangesToCheckIn() const;
 	EVisibility GetSourceControlCheckInStatusVisibility() const;
+	EVisibility GetSourceControlRestoreAsLatestVisibility() const;
 	EVisibility GetSourceControlCheckInSeparatorVisibility() const;
 	FText GetSourceControlCheckInStatusText() const;
 	FText GetSourceControlCheckInStatusTooltipText() const;
 	const FSlateBrush* GetSourceControlCheckInStatusIcon() const;
 	FReply OnSourceControlCheckInChangesClicked() const;
 
+	/** Restore as latest button */
+	FReply OnSourceControlRestoreAsLatestClicked() const;
+
 	/** Conflicts */
 	void CheckSourceControlStatus();
 public:
 	bool AreConflictsRemaining() const;
 	int32 GetNumConflictsRemaining() const;
+
+	static void SetRewoundMode(bool InRewound) { bRewoundMode = InRewound; }
 	
-	static void SetCheckInChangesDisabledOverride(bool InEnabled) { bStaticDisableCheckInChangesOverride = InEnabled; }
 	static void SetSyncLatestDisabledOverride(bool InEnabled) { bStaticDisableSyncLatestOverride = InEnabled; }
 	static FOnClicked& GetOnSyncLatestClickedStaticOverride() { return OnSyncLatestClickedStaticOverride; }
+	static FOnClicked& GetOnRestoreAsLatestClickedStaticOverride() { return OnRestoreAsLatestClickedStaticOverride; }
 
 private:
 	void OnSourceControlProviderChanged(ISourceControlProvider& OldProvider, ISourceControlProvider& NewProvider);
@@ -72,9 +80,11 @@ private:
 	
 	FOnClicked OnSyncLatestClicked;
 	FOnClicked OnCheckInChangesClicked;
+	FOnClicked OnRestoreAsLatestClicked;
 
 	TAttribute<bool> IsSyncLatestEnabled;
 	TAttribute<bool> IsCheckInChangesEnabled;
+	TAttribute<bool> IsRestoreAsLatestEnabled;
 
 	TAttribute<bool> IsSyncLatestSeparatorEnabled;
 	TAttribute<bool> IsCheckInChangesSeparatorEnabled;
@@ -86,9 +96,11 @@ private:
 	FDelegateHandle SourceControlProviderChangedHandle;
 	FDelegateHandle SourceControlStateChangedHandle;
 
-	static bool bStaticDisableCheckInChangesOverride;
+	static bool bRewoundMode;
+
 	static bool bStaticDisableSyncLatestOverride;
 	static FOnClicked OnSyncLatestClickedStaticOverride;
+	static FOnClicked OnRestoreAsLatestClickedStaticOverride;
 };
 
 #endif // SOURCE_CONTROL_WITH_SLATE
\ No newline at end of file
