SourceControl: Added 'unlocked' as a revision control style icon.

#rnx

diff --git a/Engine/Source/Developer/SourceControl/Private/RevisionControlStyle/RevisionControlStyle.cpp b/Engine/Source/Developer/SourceControl/Private/RevisionControlStyle/RevisionControlStyle.cpp
--- a/Engine/Source/Developer/SourceControl/Private/RevisionControlStyle/RevisionControlStyle.cpp
+++ b/Engine/Source/Developer/SourceControl/Private/RevisionControlStyle/RevisionControlStyle.cpp
@@ -96,10 +96,12 @@ FDefaultRevisionControlStyle::FDefaultRevisionControlStyle() : FSlateStyleSet(St
 	Set("RevisionControl.Actions.Refresh", new IMAGE_BRUSH_SVG("Starship/Common/Update", CoreStyleConstants::Icon16x16));
 	Set("RevisionControl.Actions.ChangeSettings", new IMAGE_BRUSH_SVG("Starship/SourceControl/icon_SCC_Change_Source_Control_Settings", CoreStyleConstants::Icon16x16));
 	Set("RevisionControl.Actions.Rewind", new IMAGE_BRUSH_SVG("Starship/SourceControl/RC_Rewind", CoreStyleConstants::Icon16x16));
-	Set("RevisionControl.Locked", new IMAGE_BRUSH_SVG("Starship/Common/lock", CoreStyleConstants::Icon16x16));
 
 	// Icons representing the various revision control states
-	
+
+	Set("RevisionControl.Locked", new IMAGE_BRUSH_SVG("Starship/Common/lock", CoreStyleConstants::Icon16x16));
+	Set("RevisionControl.Unlocked", new IMAGE_BRUSH_SVG("Starship/Common/lock-unlocked", CoreStyleConstants::Icon16x16));
+
 	Set("RevisionControl.CheckedOut", new IMAGE_BRUSH_SVG("Starship/SourceControl/SCC_CheckedOut", CoreStyleConstants::Icon16x16, FStyleColors::Error));
 	
 	Set("RevisionControl.OpenForAdd", new IMAGE_BRUSH_SVG("Starship/SourceControl/RC_MarkedForAdd", CoreStyleConstants::Icon16x16, FStyleColors::AccentGreen));
