SkeinUE: Changed 'conflict resolution tab icon' to grey instead of yellow to match the rest of the editor.

diff --git a/Engine/Source/Developer/SourceControl/Private/RevisionControlStyle/RevisionControlStyle.cpp b/Engine/Source/Developer/SourceControl/Private/RevisionControlStyle/RevisionControlStyle.cpp
--- a/Engine/Source/Developer/SourceControl/Private/RevisionControlStyle/RevisionControlStyle.cpp
+++ b/Engine/Source/Developer/SourceControl/Private/RevisionControlStyle/RevisionControlStyle.cpp
@@ -126,6 +126,8 @@ FDefaultRevisionControlStyle::FDefaultRevisionControlStyle() : FSlateStyleSet(St
 	
 	// Misc Icons
 	Set("RevisionControl.ChangelistsTab", new IMAGE_BRUSH_SVG("Starship/Common/check-circle", CoreStyleConstants::Icon16x16));
+	Set("RevisionControl.ConflictResolutionTab", new IMAGE_BRUSH_SVG("Starship/SourceControl/RC_Conflicted", CoreStyleConstants::Icon16x16));
+	Set("RevisionControl.SnapshotHistoryTab", new IMAGE_BRUSH_SVG("Starship/SourceControl/RC_Rewind", CoreStyleConstants::Icon16x16));
 	
 	Set("RevisionControl.StatusBar.AtLatestRevision", new IMAGE_BRUSH_SVG("Starship/SourceControl/RC_StatusRemoteUpToDate", CoreStyleConstants::Icon16x16));
 	Set("RevisionControl.StatusBar.NotAtLatestRevision", new IMAGE_BRUSH_SVG("Starship/SourceControl/RC_StatusRemoteDownload", CoreStyleConstants::Icon16x16, FStyleColors::Warning));
@@ -136,7 +138,6 @@ FDefaultRevisionControlStyle::FDefaultRevisionControlStyle() : FSlateStyleSet(St
 	
 	Set("RevisionControl.Promote.Large", new IMAGE_BRUSH_SVG("Starship/SourceControl/RC_Promote_Large", CoreStyleConstants::Icon32x32));
 
-
 	Set("RevisionControl.ConflictResolution.OpenExternal", new IMAGE_BRUSH_SVG("Starship/SourceControl/RC_ConflictResolution_OpenExternal", CoreStyleConstants::Icon16x16));
 	Set("RevisionControl.ConflictResolution.Clear", new IMAGE_BRUSH_SVG("Starship/SourceControl/RC_ConflictResolution_Clear", CoreStyleConstants::Icon16x16));
 
