Snapshot History Design/UI Fixes

* Updated dropdown text color for dates
* Added missing icon(s) for project size
* When a Snapshot is expanded, the background color now remains dark
* When two adjacent Snapshots are expanded, there is now a separation between the expanded areas

diff --git a/Engine/Source/Developer/SourceControl/Private/RevisionControlStyle/RevisionControlStyle.cpp b/Engine/Source/Developer/SourceControl/Private/RevisionControlStyle/RevisionControlStyle.cpp
--- a/Engine/Source/Developer/SourceControl/Private/RevisionControlStyle/RevisionControlStyle.cpp
+++ b/Engine/Source/Developer/SourceControl/Private/RevisionControlStyle/RevisionControlStyle.cpp
@@ -153,6 +153,7 @@ FDefaultRevisionControlStyle::FDefaultRevisionControlStyle() : FSlateStyleSet(St
 	Set("RevisionControl.ConflictedState", new IMAGE_BRUSH_SVG("Starship/SourceControl/RC_ConflictedState", CoreStyleConstants::Icon26x26));
 
 	Set("RevisionControl.File", new IMAGE_BRUSH_SVG("Starship/SourceControl/RC_File", CoreStyleConstants::Icon6x8));
+	Set("RevisionControl.DiskSize", new IMAGE_BRUSH_SVG("Starship/SourceControl/RC_DiskSize", CoreStyleConstants::Icon8x8));
 	Set("RevisionControl.State.Added", new IMAGE_BRUSH_SVG("Starship/SourceControl/RC_Added", CoreStyleConstants::Icon12x12));
 	Set("RevisionControl.State.Modified", new IMAGE_BRUSH_SVG("Starship/SourceControl/RC_Modified", CoreStyleConstants::Icon12x12));
 	Set("RevisionControl.State.Removed", new IMAGE_BRUSH_SVG("Starship/SourceControl/RC_Removed", CoreStyleConstants::Icon12x12));
