SourceControl: Fixed crash in FSourceControlAssetDataCache::LaunchFetchAssetDataTasks().

The fix is theorycrafted by reasoning about the code - I was not able to actually reproduce this.

The callstack suggests a nullptr exception, indicating that the Find() in AssetDataCache failed. The AssetDataToFetch container is populated in two places:
- AddAssetInformationEntry()
- OnUpdateHistoryComplete()

The AddAssetInformationEntry() method populates the AssetDataCache and optionally the AssetDataToFetch container.
The OnUpdateHistoryComplete() assumes the entry in AssetDataCache exists and populates the AssetDataToFetch container.

The ClearPendingTasks() removes pending entries from the AssetDataCache container.
However, it doesn't wait for any pending FUpdateStatus task to be finished.

Thus, if the ClearPendingTasks() method is called after GetFileHistory() queues an FUpdateStatus operation but before that operation returns, the OnUpdateHistoryComplete() would queue an entry in AssetDataToFetch while there is no accompanying entry in AssetDataCache. This will result in the crash. The ClearPendingTasks() is called on shutdown, when the source control provider changes or when the source control dialog is shown - the latter being the most likely thing to have happened.

diff --git a/Engine/Source/Developer/SourceControl/Private/SourceControlAssetDataCache.cpp b/Engine/Source/Developer/SourceControl/Private/SourceControlAssetDataCache.cpp
--- a/Engine/Source/Developer/SourceControl/Private/SourceControlAssetDataCache.cpp
+++ b/Engine/Source/Developer/SourceControl/Private/SourceControlAssetDataCache.cpp
@@ -182,20 +182,19 @@ void FSourceControlAssetDataCache::OnUpdateHistoryComplete(const TArray<FString>
 
 	for (const FString& Filename : InUpdatedFiles)
 	{
-		FSourceControlStatePtr FileState = SourceControlProvider.GetState(Filename, EStateCacheUsage::Use);
-		check(FileState.IsValid());
-
-		if (FileState->GetHistorySize() > 0)
+		// The AssetDataCache container entry might have been removed during FUpdateStatus.
+		if (FSourceControlAssetDataEntry* AssetDataEntry = AssetDataCache.Find(Filename))
 		{
-			AssetDataToFetch.Enqueue(FileState);
-		}
-		else
-		{
-			// History was not updated we probably cannot go further with this asset.
-			FSourceControlAssetDataEntry* AssetDataEntry = AssetDataCache.Find(Filename);
+			FSourceControlStatePtr FileState = SourceControlProvider.GetState(Filename, EStateCacheUsage::Use);
+			check(FileState.IsValid());
 
-			if (AssetDataEntry != nullptr)
+			if (FileState->GetHistorySize() > 0)
+			{
+				AssetDataToFetch.Enqueue(FileState);
+			}
+			else
 			{
+				// History was not updated we probably cannot go further with this asset.
 				AssetDataEntry->bInitialized = true;
 			}
 		}
