Warn the user with a confirmation dialog if the revision control provider is changed/disabled

diff --git a/Engine/Source/Developer/SourceControl/Private/SSourceControlLogin.cpp b/Engine/Source/Developer/SourceControl/Private/SSourceControlLogin.cpp
--- a/Engine/Source/Developer/SourceControl/Private/SSourceControlLogin.cpp
+++ b/Engine/Source/Developer/SourceControl/Private/SSourceControlLogin.cpp
@@ -190,8 +190,23 @@ FReply SSourceControlLogin::OnAcceptSettings()
 	return FReply::Handled();
 }
 
+bool SSourceControlLogin::ConfirmDisableSourceControl() const
+{
+	FSourceControlModule& SourceControlModule = FSourceControlModule::Get();
+	if (SourceControlModule.GetSourceControlProviderChanging().IsBound())
+	{
+		return SourceControlModule.GetSourceControlProviderChanging().Execute();
+	}
+	return true;
+}
+
 FReply SSourceControlLogin::OnDisableSourceControl()
 {
+	if (ConfirmDisableSourceControl() == false) // The user has decided to abort the operation
+	{
+		return FReply::Handled();
+	}
+
 	FSourceControlModule& SourceControlModule = FSourceControlModule::Get();
 	SourceControlModule.SetProvider("None");
 	if(ParentWindowPtr.IsValid())
diff --git a/Engine/Source/Developer/SourceControl/Private/SSourceControlLogin.h b/Engine/Source/Developer/SourceControl/Private/SSourceControlLogin.h
--- a/Engine/Source/Developer/SourceControl/Private/SSourceControlLogin.h
+++ b/Engine/Source/Developer/SourceControl/Private/SSourceControlLogin.h
@@ -89,6 +89,9 @@ private:
 	/** Delegate to determine visibility of the disabled text widget */
 	EVisibility GetDisabledTextVisibility() const;
 
+	/** Requests confirmation of disabling source control */
+	bool ConfirmDisableSourceControl() const;
+
 private:
 	/** The frequency at which to tick to scc module when inside a modal window */
 	static const float RefreshFrequency;
diff --git a/Engine/Source/Developer/SourceControl/Private/SSourceControlPicker.cpp b/Engine/Source/Developer/SourceControl/Private/SSourceControlPicker.cpp
--- a/Engine/Source/Developer/SourceControl/Private/SSourceControlPicker.cpp
+++ b/Engine/Source/Developer/SourceControl/Private/SSourceControlPicker.cpp
@@ -52,6 +52,15 @@ void SSourceControlPicker::Construct(const FArguments& InArgs)
 void SSourceControlPicker::ChangeSourceControlProvider(int32 ProviderIndex) const
 {
 	FSourceControlModule& SourceControlModule = FSourceControlModule::Get();
+
+	if (CurrentProviderIndex != ProviderIndex)
+	{
+		if (ConfirmProviderChanging() == false) // The user has decided to abort the operation
+		{
+			return;
+		}
+	}
+
 	SourceControlModule.SetCurrentSourceControlProvider(ProviderIndex);
 
 	if(SourceControlModule.GetLoginWidget().IsValid())
@@ -60,7 +69,17 @@ void SSourceControlPicker::ChangeSourceControlProvider(int32 ProviderIndex) cons
 	}
 }
 
-TSharedRef<SWidget> SSourceControlPicker::OnGetMenuContent() const
+bool SSourceControlPicker::ConfirmProviderChanging() const
+{
+	FSourceControlModule& SourceControlModule = FSourceControlModule::Get();
+	if (SourceControlModule.GetSourceControlProviderChanging().IsBound())
+	{
+		return SourceControlModule.GetSourceControlProviderChanging().Execute();
+	}
+	return true;
+}
+
+TSharedRef<SWidget> SSourceControlPicker::OnGetMenuContent()
 {
 	FSourceControlModule& SourceControlModule = FSourceControlModule::Get();
 
@@ -75,6 +94,11 @@ TSharedRef<SWidget> SSourceControlPicker::OnGetMenuContent() const
 		const FName ProviderName = SourceControlModule.GetSourceControlProviderName(ProviderIndex);
 		int32 ProviderSortKey = ProviderName == FName("None") ? -1 * ProviderIndex : ProviderIndex;
 		SortedProviderNames.Emplace(ProviderName, ProviderSortKey);
+
+		if (ProviderName == FSourceControlModule::Get().GetProvider().GetName())
+		{
+			CurrentProviderIndex = ProviderIndex;
+		}
 	}
 
 	// Sort based on the provider index
@@ -116,7 +140,7 @@ FText SSourceControlPicker::GetProviderText(const FName& InName) const
 		return LOCTEXT("NoProviderDescription", "None  (revision control disabled)");
 	}
 
-	// @todo: Remove this block after the Git/Skein plugins have been exhaustively tested (also remember to change the Git plugin's "IsBetaVersion" setting to false.)
+	// @todo: Remove this block after the Git plugin has been exhaustively tested (also remember to change the Git plugin's "IsBetaVersion" setting to false.)
 	if(InName == "Git")
 	{
 		return LOCTEXT( "GitBetaProviderName", "Git  (beta version)" );
diff --git a/Engine/Source/Developer/SourceControl/Private/SSourceControlPicker.h b/Engine/Source/Developer/SourceControl/Private/SSourceControlPicker.h
--- a/Engine/Source/Developer/SourceControl/Private/SSourceControlPicker.h
+++ b/Engine/Source/Developer/SourceControl/Private/SSourceControlPicker.h
@@ -32,13 +32,19 @@ private:
 	void ChangeSourceControlProvider(int32 ProviderIndex) const;
 
 	/** Get the content for the drop-down menu for picking providers */
-	TSharedRef<SWidget> OnGetMenuContent() const;
+	TSharedRef<SWidget> OnGetMenuContent();
 
 	/** Get the button text for the drop-down */
 	FText OnGetButtonText() const;
 
 	/** Get the text to be displayed given the name of the provider */
 	FText GetProviderText(const FName& InName) const;
+
+	/** Requests confirmation for the change of provider */
+	bool ConfirmProviderChanging() const;
+
+	/** The index of the currently active provider */
+	int CurrentProviderIndex = -1;
 };
 
 #endif // SOURCE_CONTROL_WITH_SLATE
diff --git a/Engine/Source/Developer/SourceControl/Private/SourceControlModule.cpp b/Engine/Source/Developer/SourceControl/Private/SourceControlModule.cpp
--- a/Engine/Source/Developer/SourceControl/Private/SourceControlModule.cpp
+++ b/Engine/Source/Developer/SourceControl/Private/SourceControlModule.cpp
@@ -481,6 +481,11 @@ void FSourceControlModule::SetUseGlobalSettings(bool bIsUseGlobalSettings)
 	ShowLoginDialog(FSourceControlLoginClosed(), ELoginWindowMode::Modeless, EOnLoginWindowStartup::PreserveProvider);
 }	
 
+FSourceControlProviderChanging& FSourceControlModule::GetSourceControlProviderChanging()
+{
+	return OnSourceControlProviderChanging;
+}
+
 FDelegateHandle FSourceControlModule::RegisterProviderChanged(const FSourceControlProviderChanged::FDelegate& SourceControlProviderChanged)
 {
 	return OnSourceControlProviderChanged.Add(SourceControlProviderChanged);
diff --git a/Engine/Source/Developer/SourceControl/Private/SourceControlModule.h b/Engine/Source/Developer/SourceControl/Private/SourceControlModule.h
--- a/Engine/Source/Developer/SourceControl/Private/SourceControlModule.h
+++ b/Engine/Source/Developer/SourceControl/Private/SourceControlModule.h
@@ -41,6 +41,7 @@ public:
 	virtual void ShowLoginDialog(const FSourceControlLoginClosed& InOnSourceControlLoginClosed, ELoginWindowMode::Type InLoginWindowMode, EOnLoginWindowStartup::Type InOnLoginWindowStartup = EOnLoginWindowStartup::ResetProviderToNone) override;
 	virtual bool GetUseGlobalSettings() const override;
 	virtual void SetUseGlobalSettings(bool bIsUseGlobalSettings) override;
+	virtual FSourceControlProviderChanging& GetSourceControlProviderChanging() override;
 	virtual FDelegateHandle RegisterProviderChanged(const FSourceControlProviderChanged::FDelegate& SourceControlProviderChanged) override;
 	virtual void UnregisterProviderChanged(FDelegateHandle Handle) override;
 	virtual void RegisterPreSubmitDataValidation(const FSourceControlPreSubmitDataValidationDelegate& PreSubmitDataValidationDelegate) override;
@@ -141,6 +142,9 @@ private:
 	/** Active Provider name to track source control provider changes */
 	FString ActiveProviderName;
 
+	/** For notifying when the source provider is about to change */
+	FSourceControlProviderChanging OnSourceControlProviderChanging;
+
 	/** For notifying when the source provider is changed */
 	FSourceControlProviderChanged OnSourceControlProviderChanged;
 
diff --git a/Engine/Source/Developer/SourceControl/Public/ISourceControlModule.h b/Engine/Source/Developer/SourceControl/Public/ISourceControlModule.h
--- a/Engine/Source/Developer/SourceControl/Public/ISourceControlModule.h
+++ b/Engine/Source/Developer/SourceControl/Public/ISourceControlModule.h
@@ -21,6 +21,9 @@ SOURCECONTROL_API DECLARE_LOG_CATEGORY_EXTERN(LogSourceControl, Log, All);
 /** Delegate called when the source control login window is closed. Parameter determines if source control is enabled or not */
 DECLARE_DELEGATE_OneParam( FSourceControlLoginClosed, bool );
 
+/** Delegate called when the active source control provider is changing. Returns false if the operation is to be aborted */
+DECLARE_DELEGATE_RetVal(bool, FSourceControlProviderChanging);
+
 /** Delegate called when the active source control provider is changed */
 DECLARE_MULTICAST_DELEGATE_TwoParams( FSourceControlProviderChanged, ISourceControlProvider& /*OldProvider*/, ISourceControlProvider& /*NewProvider*/ );
 
@@ -177,6 +180,11 @@ public:
 	 */
 	virtual void SetUseGlobalSettings(bool bIsUseGlobalSettings) = 0;
 
+	/**
+	 * Retrieve a reference to a delegate to be called before the source control provider changes
+	 */
+	virtual FSourceControlProviderChanging& GetSourceControlProviderChanging() = 0;
+	
 	/**
 	 * Register a delegate to be called when the source control provider changes
 	 */
