Skein / ConflictUE: Update primary URC buttons with conflict icons when a user saves changes to a file that we know is out of date.

diff --git a/Engine/Source/Developer/SourceControl/Private/SSourceControlControls.cpp b/Engine/Source/Developer/SourceControl/Private/SSourceControlControls.cpp
--- a/Engine/Source/Developer/SourceControl/Private/SSourceControlControls.cpp
+++ b/Engine/Source/Developer/SourceControl/Private/SSourceControlControls.cpp
@@ -14,6 +14,7 @@
 #define LOCTEXT_NAMESPACE "SSkeinSourceControlWidgets"
 
 int32 SSourceControlControls::NumConflictsRemaining = 0;
+int32 SSourceControlControls::NumConflictsUpcoming = 0;
 
 FIsEnabled SSourceControlControls::IsSyncLatestEnabled;
 FIsEnabled SSourceControlControls::IsCheckInChangesEnabled;
@@ -204,14 +205,29 @@ void SSourceControlControls::OnSourceControlStateChanged()
 {
 	ISourceControlProvider& SourceControlProvider = ISourceControlModule::Get().GetProvider();
 
-	TArray<FSourceControlStateRef> Conflicts = SourceControlProvider.GetCachedStateByPredicate(
+	TArray<FSourceControlStateRef> ConflictsRemaining = SourceControlProvider.GetCachedStateByPredicate(
 		[](const FSourceControlStateRef& State)
 		{
 			return State->IsConflicted();
 		}
 	);
 
-	NumConflictsRemaining = Conflicts.Num(); // Atomic write.
+	TArray<FSourceControlStateRef> ConflictsUpcoming = SourceControlProvider.GetCachedStateByPredicate(
+		[](const FSourceControlStateRef& State)
+		{
+			if (State->IsCurrent())
+			{
+				return false;
+			}
+			else
+			{
+				return State->IsModified() || State->IsDeleted() || State->IsAdded();
+			}
+		}
+	);
+
+	NumConflictsRemaining = ConflictsRemaining.Num(); // Atomic write.
+	NumConflictsUpcoming = ConflictsUpcoming.Num(); // Atomic write.
 }
 
 int32 SSourceControlControls::GetNumConflictsRemaining()
@@ -219,6 +235,11 @@ int32 SSourceControlControls::GetNumConflictsRemaining()
 	return NumConflictsRemaining;
 }
 
+int32 SSourceControlControls::GetNumConflictsUpcoming()
+{
+	return NumConflictsUpcoming;
+}
+
 /** Sync Status */
 
 bool SSourceControlControls::IsAtLatestRevision()
@@ -308,6 +329,10 @@ FText SSourceControlControls::GetSourceControlSyncStatusToolTipText()
 	{
 		return LOCTEXT("SyncLatestButtonNotAtHeadTooltipTextConflict", "Some of your local changes conflict with the latest snapshot of the project. Click here to review these conflicts.");
 	}
+	if (GetNumConflictsUpcoming() > 0)
+	{
+		return LOCTEXT("SyncLatestButtonNotAtHeadTooltipTextConflictUpcoming", "Some of your local changes conflict with the latest snapshot of the project. Click here to review these conflicts.");
+	}
 	if (HasSourceControlChangesToSync())
 	{
 		return LOCTEXT("SyncLatestButtonNotAtHeadTooltipText", "Sync to the latest Snapshot for this project");
@@ -326,6 +351,10 @@ const FSlateBrush* SSourceControlControls::GetSourceControlSyncStatusIcon()
 	{
 		return ConflictBrush;
 	}
+	if (GetNumConflictsUpcoming() > 0)
+	{
+		return ConflictBrush;
+	}
 	if (HasSourceControlChangesToSync())
 	{
 		return NotAtHeadBrush;
@@ -447,6 +476,10 @@ FText SSourceControlControls::GetSourceControlCheckInStatusToolTipText()
 	{
 		return LOCTEXT("CheckInButtonChangesTooltipTextConflict", "Some of your local changes conflict with the latest snapshot of the project. Click here to review these conflicts.");
 	}
+	if (GetNumConflictsUpcoming() > 0)
+	{
+		return LOCTEXT("CheckInButtonChangesTooltipTextConflictUpcoming", "Some of your local changes conflict with the latest snapshot of the project. Click here to review these conflicts.");
+	}
 	if (HasSourceControlChangesToCheckIn())
 	{
 		return FText::Format(LOCTEXT("CheckInButtonChangesTooltipText", "Check-in {0} change(s) to this project"), GetNumLocalChanges());
@@ -465,6 +498,10 @@ const FSlateBrush* SSourceControlControls::GetSourceControlCheckInStatusIcon()
 	{
 		return ConflictBrush;
 	}
+	if (GetNumConflictsUpcoming() > 0)
+	{
+		return ConflictBrush;
+	}
 	if (HasSourceControlChangesToCheckIn())
 	{
 		return HasLocalChangesBrush;
diff --git a/Engine/Source/Developer/SourceControl/Public/SSourceControlControls.h b/Engine/Source/Developer/SourceControl/Public/SSourceControlControls.h
--- a/Engine/Source/Developer/SourceControl/Public/SSourceControlControls.h
+++ b/Engine/Source/Developer/SourceControl/Public/SSourceControlControls.h
@@ -63,6 +63,7 @@ public:
 
 public:
 	static int32 GetNumConflictsRemaining();
+	static int32 GetNumConflictsUpcoming();
 
 public:
 	static void SetIsSyncLatestEnabled(const FIsEnabled& InSyncLatestEnabled) { IsSyncLatestEnabled = InSyncLatestEnabled; }
@@ -92,6 +93,7 @@ private:
 	FDelegateHandle SourceControlStateChangedHandle;
 
 	static int32 NumConflictsRemaining;
+	static int32 NumConflictsUpcoming;
 
 	static FIsEnabled IsSyncLatestEnabled;
 	static FIsEnabled IsCheckInChangesEnabled;
