Skein / ViewportStatus: Update iconography again with latest icons/colors received on Slack.

diff --git a/Engine/Plugins/Experimental/TypedElementsDataStorage/Source/TypedElementsDataStorageRevisionControl/Source/Private/RevisionControlProcessors.cpp b/Engine/Plugins/Experimental/TypedElementsDataStorage/Source/TypedElementsDataStorageRevisionControl/Source/Private/RevisionControlProcessors.cpp
--- a/Engine/Plugins/Experimental/TypedElementsDataStorage/Source/TypedElementsDataStorageRevisionControl/Source/Private/RevisionControlProcessors.cpp
+++ b/Engine/Plugins/Experimental/TypedElementsDataStorage/Source/TypedElementsDataStorageRevisionControl/Source/Private/RevisionControlProcessors.cpp
@@ -73,7 +73,7 @@ static FColor DetermineOverlayColor(const TypedElementDataStorage::IQueryContext
 			if (CVarEnableOverlayNotAtHeadRevision.GetValueOnGameThread())
 			{
 				// Yellow.
-				return FColor(255, 255, 61, CVarOverlayAlpha.GetValueOnGameThread());
+				return FColor(225, 255, 61, CVarOverlayAlpha.GetValueOnGameThread());
 			}
 		}
 
@@ -83,7 +83,7 @@ static FColor DetermineOverlayColor(const TypedElementDataStorage::IQueryContext
 			if (CVarEnableOverlayCheckedOutByOtherUser.GetValueOnGameThread())
 			{
 				// Red.
-				return FColor(255, 64, 64, CVarOverlayAlpha.GetValueOnGameThread());
+				return FColor(239, 53, 53, CVarOverlayAlpha.GetValueOnGameThread());
 			}
 		}
 
@@ -96,8 +96,8 @@ static FColor DetermineOverlayColor(const TypedElementDataStorage::IQueryContext
 				{
 					if (StatusColumn->Modification == ESCCModification::Added)
 					{
-						// Green.
-						return FColor(134, 194, 74, CVarOverlayAlpha.GetValueOnGameThread());
+						// Blue.
+						return FColor(0, 112, 224, CVarOverlayAlpha.GetValueOnGameThread());
 					}
 				}
 			}
@@ -108,8 +108,8 @@ static FColor DetermineOverlayColor(const TypedElementDataStorage::IQueryContext
 		{
 			if (CVarEnableOverlayCheckedOut.GetValueOnGameThread())
 			{
-				// Blue.
-				return FColor(0, 112, 224, CVarOverlayAlpha.GetValueOnGameThread());
+				// Green.
+				return FColor(31, 228, 75, CVarOverlayAlpha.GetValueOnGameThread());
 			}
 		}
 	}
diff --git a/Engine/Source/Developer/SourceControl/Private/RevisionControlStyle/RevisionControlStyle.cpp b/Engine/Source/Developer/SourceControl/Private/RevisionControlStyle/RevisionControlStyle.cpp
--- a/Engine/Source/Developer/SourceControl/Private/RevisionControlStyle/RevisionControlStyle.cpp
+++ b/Engine/Source/Developer/SourceControl/Private/RevisionControlStyle/RevisionControlStyle.cpp
@@ -76,6 +76,11 @@ FDefaultRevisionControlStyle::FDefaultRevisionControlStyle() : FSlateStyleSet(St
 	// We use a custom teal color for the branched icons
 	BranchedColor = FLinearColor::FromSRGBColor(FColor::FromHex("#00E4A0"));
 
+	StatusCheckedOutColor = FLinearColor::FromSRGBColor(FColor::FromHex("#1FE44B"));
+	StatusCheckedOutByOtherUserColor = FLinearColor::FromSRGBColor(FColor::FromHex("#EF3535"));
+	StatusNotAtHeadRevisionColor = FLinearColor::FromSRGBColor(FColor::FromHex("#E1FF3D"));
+	StatusOpenForAddColor = FLinearColor::FromSRGBColor(FColor::FromHex("#0070E0"));
+
 	// Status icons
 	Set("RevisionControl.Icon", new IMAGE_BRUSH_SVG("Starship/SourceControl/Status/RevisionControl", CoreStyleConstants::Icon16x16));
 	Set("RevisionControl.Icon.ConnectedBadge", new IMAGE_BRUSH_SVG("Starship/SourceControl/Status/RevisionControlBadgeConnected", CoreStyleConstants::Icon16x16, FStyleColors::Success));
@@ -102,12 +107,12 @@ FDefaultRevisionControlStyle::FDefaultRevisionControlStyle() : FSlateStyleSet(St
 	Set("RevisionControl.Locked", new IMAGE_BRUSH_SVG("Starship/Common/lock", CoreStyleConstants::Icon16x16));
 	Set("RevisionControl.Unlocked", new IMAGE_BRUSH_SVG("Starship/Common/lock-unlocked", CoreStyleConstants::Icon16x16));
 
-	Set("RevisionControl.CheckedOut", new IMAGE_BRUSH_SVG("Starship/SourceControl/SCC_CheckedOut", CoreStyleConstants::Icon16x16, FStyleColors::AccentBlue));
+	Set("RevisionControl.CheckedOut", new IMAGE_BRUSH_SVG("Starship/SourceControl/SCC_CheckedOut", CoreStyleConstants::Icon16x16, StatusCheckedOutColor));
 	
-	Set("RevisionControl.OpenForAdd", new IMAGE_BRUSH_SVG("Starship/SourceControl/RC_MarkedForAdd", CoreStyleConstants::Icon16x16, FStyleColors::AccentGreen));
+	Set("RevisionControl.OpenForAdd", new IMAGE_BRUSH_SVG("Starship/SourceControl/RC_MarkedForAdd", CoreStyleConstants::Icon16x16, StatusOpenForAddColor));
 	Set("RevisionControl.MarkedForDelete", new IMAGE_BRUSH_SVG("Starship/SourceControl/SCC_MarkedForDelete", CoreStyleConstants::Icon16x16, FStyleColors::Error));
 
-	Set("RevisionControl.CheckedOutByOtherUser", new IMAGE_BRUSH_SVG("Starship/SourceControl/RC_CheckedOther", CoreStyleConstants::Icon16x16, FStyleColors::Error));
+	Set("RevisionControl.CheckedOutByOtherUser", new IMAGE_BRUSH_SVG("Starship/SourceControl/RC_CheckedOther", CoreStyleConstants::Icon16x16, StatusCheckedOutByOtherUserColor));
 	Set("RevisionControl.CheckedOutByOtherUserBadge", new IMAGE_BRUSH_SVG("Starship/SourceControl/RC_CheckedOtherBadge", CoreStyleConstants::Icon16x16, FStyleColors::AccentBlue));
 
 	Set("RevisionControl.CheckedOutByOtherUserOtherBranch", new IMAGE_BRUSH_SVG("Starship/SourceControl/RC_CheckedBranch", CoreStyleConstants::Icon16x16, BranchedColor));
@@ -117,7 +122,7 @@ FDefaultRevisionControlStyle::FDefaultRevisionControlStyle() : FSlateStyleSet(St
 	Set("RevisionControl.ModifiedBadge", new IMAGE_BRUSH_SVG("Starship/SourceControl/RC_BranchModifiedBadge", CoreStyleConstants::Icon16x16, FStyleColors::Warning));
 	Set("RevisionControl.ModifiedLocally", new IMAGE_BRUSH_SVG("Starship/SourceControl/RC_ModifiedLocally", CoreStyleConstants::Icon16x16, FStyleColors::Foreground));
 
-	Set("RevisionControl.NotAtHeadRevision", new IMAGE_BRUSH_SVG("Starship/SourceControl/RC_NewerVersion", CoreStyleConstants::Icon16x16, FStyleColors::Warning));
+	Set("RevisionControl.NotAtHeadRevision", new IMAGE_BRUSH_SVG("Starship/SourceControl/RC_NewerVersion", CoreStyleConstants::Icon16x16, StatusNotAtHeadRevisionColor));
 	Set("RevisionControl.NotInDepot", new IMAGE_BRUSH_SVG("Starship/SourceControl/SCC_NotInDepot", CoreStyleConstants::Icon16x16, FStyleColors::Warning));
 	
 	Set("RevisionControl.Branched", new IMAGE_BRUSH_SVG("Starship/SourceControl/SCC_Branched", CoreStyleConstants::Icon16x16, BranchedColor));
@@ -130,7 +135,7 @@ FDefaultRevisionControlStyle::FDefaultRevisionControlStyle() : FSlateStyleSet(St
 	Set("RevisionControl.SnapshotHistoryTab", new IMAGE_BRUSH_SVG("Starship/SourceControl/RC_Rewind", CoreStyleConstants::Icon16x16));
 	
 	Set("RevisionControl.StatusBar.AtLatestRevision", new IMAGE_BRUSH_SVG("Starship/SourceControl/RC_StatusRemoteUpToDate", CoreStyleConstants::Icon16x16));
-	Set("RevisionControl.StatusBar.NotAtLatestRevision", new IMAGE_BRUSH_SVG("Starship/SourceControl/RC_StatusRemoteDownload", CoreStyleConstants::Icon16x16, FStyleColors::Warning));
+	Set("RevisionControl.StatusBar.NotAtLatestRevision", new IMAGE_BRUSH_SVG("Starship/SourceControl/RC_StatusRemoteDownload", CoreStyleConstants::Icon16x16, StatusNotAtHeadRevisionColor));
 	Set("RevisionControl.StatusBar.NoLocalChanges", new IMAGE_BRUSH_SVG("Starship/SourceControl/RC_StatusLocalUpToDate", CoreStyleConstants::Icon16x16));
 	Set("RevisionControl.StatusBar.HasLocalChanges", new IMAGE_BRUSH_SVG("Starship/SourceControl/RC_StatusLocalUpload", CoreStyleConstants::Icon16x16, FStyleColors::AccentBlue));
 	Set("RevisionControl.StatusBar.Conflicted", new IMAGE_BRUSH_SVG("Starship/SourceControl/RC_Conflicted", CoreStyleConstants::Icon16x16, FStyleColors::Warning));
diff --git a/Engine/Source/Developer/SourceControl/Public/RevisionControlStyle/RevisionControlStyle.h b/Engine/Source/Developer/SourceControl/Public/RevisionControlStyle/RevisionControlStyle.h
--- a/Engine/Source/Developer/SourceControl/Public/RevisionControlStyle/RevisionControlStyle.h
+++ b/Engine/Source/Developer/SourceControl/Public/RevisionControlStyle/RevisionControlStyle.h
@@ -56,6 +56,12 @@ protected:
 
 	/** The specific color we use for all the "Branched" icons */
 	FLinearColor BranchedColor;
+
+	/** The specific colors we use for all the "Status" icons */
+	FLinearColor StatusCheckedOutColor;
+	FLinearColor StatusCheckedOutByOtherUserColor;
+	FLinearColor StatusNotAtHeadRevisionColor;
+	FLinearColor StatusOpenForAddColor;
 	
 private:
 	static FName StyleName;
