Source Control - File Status Monitor: Prevent crash when accessing Slate in unattended apps

diff --git a/Engine/Source/Developer/SourceControl/Private/SourceControlFileStatusMonitor.cpp b/Engine/Source/Developer/SourceControl/Private/SourceControlFileStatusMonitor.cpp
--- a/Engine/Source/Developer/SourceControl/Private/SourceControlFileStatusMonitor.cpp
+++ b/Engine/Source/Developer/SourceControl/Private/SourceControlFileStatusMonitor.cpp
@@ -6,6 +6,7 @@
 #include "ISourceControlModule.h"
 #include "ISourceControlProvider.h"
 #include "Math/NumericLimits.h"
+#include "Misc/App.h"
 #include "Misc/ScopeExit.h"
 #include "ProfilingDebugging/CpuProfilerTrace.h"
 #include "SourceControlOperations.h"
@@ -22,12 +23,15 @@ FSourceControlFileStatusMonitor::FSourceControlFileStatusMonitor()
 	SetSuspendMonitoringPolicy([]()
 	{
 #if SOURCE_CONTROL_WITH_SLATE
-		// By default, suspend monitoring if the user didn't interact for the last 5 minutes.
-		return FPlatformTime::Seconds() - FSlateApplication::Get().GetLastUserInteractionTime() > FTimespan::FromMinutes(5).GetTotalSeconds();
-#else
+		if (!FApp::IsUnattended() && FSlateApplication::IsInitialized())
+		{
+			// By default, suspend monitoring if the user didn't interact for the last 5 minutes.
+			return FPlatformTime::Seconds() - FSlateApplication::Get().GetLastUserInteractionTime() > FTimespan::FromMinutes(5).GetTotalSeconds();
+		}
+#endif //SOURCE_CONTROL_WITH_SLATE
+
 		// Without slate there is not user interaction, so we always suspend the monitoring
 		return true;
-#endif //SOURCE_CONTROL_WITH_SLATE
 	});
 }
 
