Revision Control: Update color for marked for add icon to accent blue

diff --git a/Engine/Source/Developer/SourceControl/Private/RevisionControlStyle/RevisionControlStyle.cpp b/Engine/Source/Developer/SourceControl/Private/RevisionControlStyle/RevisionControlStyle.cpp
--- a/Engine/Source/Developer/SourceControl/Private/RevisionControlStyle/RevisionControlStyle.cpp
+++ b/Engine/Source/Developer/SourceControl/Private/RevisionControlStyle/RevisionControlStyle.cpp
@@ -79,7 +79,6 @@ FDefaultRevisionControlStyle::FDefaultRevisionControlStyle() : FSlateStyleSet(St
 	StatusCheckedOutColor = FLinearColor::FromSRGBColor(FColor::FromHex("#1FE44B"));
 	StatusCheckedOutByOtherUserColor = FLinearColor::FromSRGBColor(FColor::FromHex("#EF3535"));
 	StatusNotAtHeadRevisionColor = FLinearColor::FromSRGBColor(FColor::FromHex("#E1FF3D"));
-	StatusOpenForAddColor = FLinearColor::FromSRGBColor(FColor::FromHex("#0070E0"));
 
 	// Status icons
 	Set("RevisionControl.Icon", new IMAGE_BRUSH_SVG("Starship/SourceControl/Status/RevisionControl", CoreStyleConstants::Icon16x16));
@@ -109,7 +108,7 @@ FDefaultRevisionControlStyle::FDefaultRevisionControlStyle() : FSlateStyleSet(St
 
 	Set("RevisionControl.CheckedOut", new IMAGE_BRUSH_SVG("Starship/SourceControl/SCC_CheckedOut", CoreStyleConstants::Icon16x16, StatusCheckedOutColor));
 	
-	Set("RevisionControl.OpenForAdd", new IMAGE_BRUSH_SVG("Starship/SourceControl/RC_MarkedForAdd", CoreStyleConstants::Icon16x16, StatusOpenForAddColor));
+	Set("RevisionControl.OpenForAdd", new IMAGE_BRUSH_SVG("Starship/SourceControl/RC_MarkedForAdd", CoreStyleConstants::Icon16x16, FStyleColors::AccentBlue));
 	Set("RevisionControl.MarkedForDelete", new IMAGE_BRUSH_SVG("Starship/SourceControl/SCC_MarkedForDelete", CoreStyleConstants::Icon16x16, FStyleColors::Error));
 
 	Set("RevisionControl.CheckedOutByOtherUser", new IMAGE_BRUSH_SVG("Starship/SourceControl/RC_CheckedOther", CoreStyleConstants::Icon16x16, StatusCheckedOutByOtherUserColor));
diff --git a/Engine/Source/Developer/SourceControl/Public/RevisionControlStyle/RevisionControlStyle.h b/Engine/Source/Developer/SourceControl/Public/RevisionControlStyle/RevisionControlStyle.h
--- a/Engine/Source/Developer/SourceControl/Public/RevisionControlStyle/RevisionControlStyle.h
+++ b/Engine/Source/Developer/SourceControl/Public/RevisionControlStyle/RevisionControlStyle.h
@@ -61,7 +61,6 @@ protected:
 	FLinearColor StatusCheckedOutColor;
 	FLinearColor StatusCheckedOutByOtherUserColor;
 	FLinearColor StatusNotAtHeadRevisionColor;
-	FLinearColor StatusOpenForAddColor;
 	
 private:
 	static FName StyleName;
