SourceControl: Fixed 'Sync' option not displaying in the status bar.

I wasn't actually able to reproduce the issue, but I did find some oddities regarding the visibility of the checkin/sync/restore buttons. For example: if Skein is 'unavailable', the checkin/sync buttons would be hidden but the restore button could be available. I also made the buttons more explicitly tied to Skein to avoid them showing up using Git/Perforce.

diff --git a/Engine/Source/Developer/SourceControl/Private/SSourceControlControls.cpp b/Engine/Source/Developer/SourceControl/Private/SSourceControlControls.cpp
--- a/Engine/Source/Developer/SourceControl/Private/SSourceControlControls.cpp
+++ b/Engine/Source/Developer/SourceControl/Private/SSourceControlControls.cpp
@@ -28,6 +28,50 @@ FOnClicked SSourceControlControls::OnSyncLatestClicked;
 FOnClicked SSourceControlControls::OnCheckInChangesClicked;
 FOnClicked SSourceControlControls::OnRestoreAsLatestClicked;
 
+static bool DisplaySyncStatus()
+{
+	bool bDisplaySourceControlSyncStatus = false;
+	GConfig->GetBool(TEXT("SourceControlSettings"), TEXT("DisplaySourceControlSyncStatus"), bDisplaySourceControlSyncStatus, GEditorIni);
+
+	if (bDisplaySourceControlSyncStatus)
+	{
+		ISourceControlModule& SourceControlModule = ISourceControlModule::Get();
+		if (SourceControlModule.IsEnabled() &&
+			SourceControlModule.GetProvider().IsAvailable() &&
+			SourceControlModule.GetProvider().IsAtLatestRevision().IsSet()) // Only providers that implement IsAtLatestRevision are supported.
+		{
+			return true;
+		}
+	}
+
+	return false;
+}
+
+static bool DisplayCheckInStatus()
+{
+	bool bDisplaySourceControlCheckInStatus = false;
+	GConfig->GetBool(TEXT("SourceControlSettings"), TEXT("DisplaySourceControlCheckInStatus"), bDisplaySourceControlCheckInStatus, GEditorIni);
+
+	if (bDisplaySourceControlCheckInStatus)
+	{
+		ISourceControlModule& SourceControlModule = ISourceControlModule::Get();
+		if (SourceControlModule.IsEnabled() &&
+			SourceControlModule.GetProvider().IsAvailable() &&
+			SourceControlModule.GetProvider().GetNumLocalChanges().IsSet()) // Only providers that implement GetNumLocalChanges are supported.
+		{
+			return true;
+		}
+	}
+
+	return false;
+}
+
+static bool DisplayRestoreAsLatestStatus()
+{
+	// The 'Restore as Latest' button is a replacement for the 'Check in Changes' button.
+	return DisplayCheckInStatus();
+}
+
 /**
  * Construct this widget
  *
@@ -279,24 +323,16 @@ EVisibility SSourceControlControls::GetSourceControlSyncStatusVisibility()
 		return EVisibility::Visible;
 	}
 
-	bool bVisibleSourceControlSyncStatus = true;
+	bool bVisibleSourceControlSyncStatus = false;
 	if (IsSyncLatestVisible.IsBound())
 	{
 		bVisibleSourceControlSyncStatus = IsSyncLatestVisible.Execute();
 	}
 
-	bool bDisplaySourceControlSyncStatus = false;
-	GConfig->GetBool(TEXT("SourceControlSettings"), TEXT("DisplaySourceControlSyncStatus"), bDisplaySourceControlSyncStatus, GEditorIni);
-
+	bool bDisplaySourceControlSyncStatus = DisplaySyncStatus();
 	if (bVisibleSourceControlSyncStatus && bDisplaySourceControlSyncStatus)
 	{
-		ISourceControlModule& SourceControlModule = ISourceControlModule::Get();
-		if (SourceControlModule.IsEnabled() &&
-			SourceControlModule.GetProvider().IsAvailable() &&
-			SourceControlModule.GetProvider().IsAtLatestRevision().IsSet()) // Only providers that implement IsAtLatestRevision are supported.
-		{
-			return EVisibility::Visible;
-		}
+		return EVisibility::Visible;
 	}
 
 	return EVisibility::Collapsed;
@@ -426,24 +462,16 @@ EVisibility SSourceControlControls::GetSourceControlCheckInStatusVisibility()
 		return EVisibility::Visible;
 	}
 
-	bool bVisibleSourceControlCheckInStatus = true;
+	bool bVisibleSourceControlCheckInStatus = false;
 	if (IsCheckInChangesVisible.IsBound())
 	{
 		bVisibleSourceControlCheckInStatus = IsCheckInChangesVisible.Execute();
 	}
 
-	bool bDisplaySourceControlCheckInStatus = false;
-	GConfig->GetBool(TEXT("SourceControlSettings"), TEXT("DisplaySourceControlCheckInStatus"), bDisplaySourceControlCheckInStatus, GEditorIni);
-
+	bool bDisplaySourceControlCheckInStatus = DisplayCheckInStatus();
 	if (bVisibleSourceControlCheckInStatus && bDisplaySourceControlCheckInStatus)
 	{
-		ISourceControlModule& SourceControlModule = ISourceControlModule::Get();
-		if (SourceControlModule.IsEnabled() &&
-			SourceControlModule.GetProvider().IsAvailable() &&
-			SourceControlModule.GetProvider().GetNumLocalChanges().IsSet()) // Only providers that implement GetNumLocalChanges are supported.
-		{
-			return EVisibility::Visible;
-		}
+		return EVisibility::Visible;
 	}
 
 	return EVisibility::Collapsed;
@@ -544,9 +572,16 @@ bool SSourceControlControls::IsSourceControlRestoreAsLatestEnabled()
 
 EVisibility SSourceControlControls::GetSourceControlRestoreAsLatestVisibility()
 {
+	bool bVisibleSourceControlRestoreAsLatestStatus = false;
 	if (IsRestoreAsLatestVisible.IsBound())
 	{
-		return IsRestoreAsLatestVisible.Execute() ? EVisibility::Visible : EVisibility::Collapsed;
+		bVisibleSourceControlRestoreAsLatestStatus = IsRestoreAsLatestVisible.Execute();
+	}
+
+	bool bDisplaySourceControlRestoreAsLatestStatus = DisplayRestoreAsLatestStatus();
+	if (bVisibleSourceControlRestoreAsLatestStatus && bDisplaySourceControlRestoreAsLatestStatus)
+	{
+		return EVisibility::Visible;
 	}
 
 	return EVisibility::Collapsed;
