SourceControl: Fixed hitch caused by conflict remaining/upcoming scanning in large repositories.

Instead of scanning the cache in its entirity (twice!) on every source control status cache, the Skein provider now keeps track of remaining/upcoming conflicts incrementally as new source control statuses come in.

diff --git a/Engine/Source/Developer/SourceControl/Private/SSourceControlControls.cpp b/Engine/Source/Developer/SourceControl/Private/SSourceControlControls.cpp
--- a/Engine/Source/Developer/SourceControl/Private/SSourceControlControls.cpp
+++ b/Engine/Source/Developer/SourceControl/Private/SSourceControlControls.cpp
@@ -13,8 +13,8 @@
 
 #define LOCTEXT_NAMESPACE "SSkeinSourceControlWidgets"
 
-int32 SSourceControlControls::NumConflictsRemaining = 0;
-int32 SSourceControlControls::NumConflictsUpcoming = 0;
+FNumConflicts SSourceControlControls::NumConflictsRemaining;
+FNumConflicts SSourceControlControls::NumConflictsUpcoming;
 
 FIsEnabled SSourceControlControls::IsSyncLatestEnabled;
 FIsEnabled SSourceControlControls::IsCheckInChangesEnabled;
@@ -210,78 +210,16 @@ void SSourceControlControls::Construct(const FArguments& InArgs)
 			.Orientation(EOrientation::Orient_Vertical)
 		]
 	];
-
-	CheckSourceControlStatus();
-}
-
-void SSourceControlControls::CheckSourceControlStatus()
-{
-	ISourceControlModule& SourceControlModule = ISourceControlModule::Get();
-
-	if (!SourceControlProviderChangedHandle.IsValid())
-	{
-		SourceControlProviderChangedHandle = SourceControlModule.RegisterProviderChanged(
-			FSourceControlProviderChanged::FDelegate::CreateSP(this, &SSourceControlControls::OnSourceControlProviderChanged)
-		);
-		SourceControlStateChangedHandle = SourceControlModule.GetProvider().RegisterSourceControlStateChanged_Handle(
-			FSourceControlStateChanged::FDelegate::CreateSP(this, &SSourceControlControls::OnSourceControlStateChanged)
-		);
-	}
-}
-
-void SSourceControlControls::OnSourceControlProviderChanged(ISourceControlProvider& OldProvider, ISourceControlProvider& NewProvider)
-{
-	if (SourceControlStateChangedHandle.IsValid())
-	{
-		OldProvider.UnregisterSourceControlStateChanged_Handle(SourceControlStateChangedHandle);
-		SourceControlStateChangedHandle.Reset();
-	}
-
-	if (!IsEngineExitRequested())
-	{
-		SourceControlStateChangedHandle = NewProvider.RegisterSourceControlStateChanged_Handle(
-			FSourceControlStateChanged::FDelegate::CreateSP(this, &SSourceControlControls::OnSourceControlStateChanged)
-		);
-	}
-}
-
-void SSourceControlControls::OnSourceControlStateChanged()
-{
-	ISourceControlProvider& SourceControlProvider = ISourceControlModule::Get().GetProvider();
-
-	TArray<FSourceControlStateRef> ConflictsRemaining = SourceControlProvider.GetCachedStateByPredicate(
-		[](const FSourceControlStateRef& State)
-		{
-			return State->IsConflicted();
-		}
-	);
-
-	TArray<FSourceControlStateRef> ConflictsUpcoming = SourceControlProvider.GetCachedStateByPredicate(
-		[](const FSourceControlStateRef& State)
-		{
-			if (State->IsCurrent())
-			{
-				return false;
-			}
-			else
-			{
-				return State->IsModified() || State->IsDeleted() || State->IsAdded();
-			}
-		}
-	);
-
-	NumConflictsRemaining = ConflictsRemaining.Num(); // Atomic write.
-	NumConflictsUpcoming = ConflictsUpcoming.Num(); // Atomic write.
 }
 
 int32 SSourceControlControls::GetNumConflictsRemaining()
 {
-	return NumConflictsRemaining;
+	return NumConflictsRemaining.IsBound() ? NumConflictsRemaining.Execute() : 0;
 }
 
 int32 SSourceControlControls::GetNumConflictsUpcoming()
 {
-	return NumConflictsUpcoming;
+	return NumConflictsUpcoming.IsBound() ? NumConflictsUpcoming.Execute() : 0;
 }
 
 /** Sync Status */
diff --git a/Engine/Source/Developer/SourceControl/Public/SSourceControlControls.h b/Engine/Source/Developer/SourceControl/Public/SSourceControlControls.h
--- a/Engine/Source/Developer/SourceControl/Public/SSourceControlControls.h
+++ b/Engine/Source/Developer/SourceControl/Public/SSourceControlControls.h
@@ -11,6 +11,8 @@
 #include "Styling/SlateTypes.h"
 #include "Widgets/SCompoundWidget.h"
 
+DECLARE_DELEGATE_RetVal(int32, FNumConflicts);
+
 DECLARE_DELEGATE_RetVal(bool, FIsVisible);
 DECLARE_DELEGATE_RetVal(bool, FIsEnabled);
 
@@ -66,6 +68,9 @@ public:
 	static int32 GetNumConflictsUpcoming();
 
 public:
+	static void SetNumConflictsRemaining(const FNumConflicts& InNumConflictsRemaining) { NumConflictsRemaining = InNumConflictsRemaining; }
+	static void SetNumConflictsUpcoming(const FNumConflicts& InNumConflictsUpcoming) { NumConflictsUpcoming = InNumConflictsUpcoming; }
+
 	static void SetIsSyncLatestEnabled(const FIsEnabled& InSyncLatestEnabled) { IsSyncLatestEnabled = InSyncLatestEnabled; }
 	static void SetIsCheckInChangesEnabled(const FIsEnabled& InCheckInChangesEnabled) { IsCheckInChangesEnabled = InCheckInChangesEnabled; }
 	static void SetIsRestoreAsLatestEnabled(const FIsEnabled& InRestoreAsLatestEnabled) { IsRestoreAsLatestEnabled = InRestoreAsLatestEnabled; }
@@ -78,22 +83,13 @@ public:
 	static void SetOnCheckInChangesClicked(const FOnClicked& InCheckInChangesClicked) { OnCheckInChangesClicked = InCheckInChangesClicked; }
 	static void SetOnRestoreAsLatestClicked(const FOnClicked& InRestoreAsLatestClicked) { OnRestoreAsLatestClicked = InRestoreAsLatestClicked; }
 
-private:
-	void OnSourceControlProviderChanged(ISourceControlProvider& OldProvider, ISourceControlProvider& NewProvider);
-	void OnSourceControlStateChanged();
-	
-	void CheckSourceControlStatus();
-
 private:
 	
 	TAttribute<bool> IsMiddleSeparatorEnabled;
 	TAttribute<bool> IsRightSeparatorEnabled;
 
-	FDelegateHandle SourceControlProviderChangedHandle;
-	FDelegateHandle SourceControlStateChangedHandle;
-
-	static int32 NumConflictsRemaining;
-	static int32 NumConflictsUpcoming;
+	static FNumConflicts NumConflictsRemaining;
+	static FNumConflicts NumConflictsUpcoming;
 
 	static FIsEnabled IsSyncLatestEnabled;
 	static FIsEnabled IsCheckInChangesEnabled;
