Skipping Asset Data Update on the source control module when a actor label is changed if its pacakge is transient.

Some tools like the Chaos Visual Debugger need to spawn a large number of transient actors, which will result in actor label creations which becomes expensive due to that asset data update attempt. Transient packages cannot be sabed, therefore we can safely skip the update in the source control module for them.

#tests Tested in the editor that if I change the label of any actor in a saved map, the code to update the assed data cache is still called. But if the actors are from a transient world (like in the Chaos Visual Debugger) the update is no longer executed (as expected).

diff --git a/Engine/Source/Developer/SourceControl/Private/SourceControlAssetDataCache.cpp b/Engine/Source/Developer/SourceControl/Private/SourceControlAssetDataCache.cpp
--- a/Engine/Source/Developer/SourceControl/Private/SourceControlAssetDataCache.cpp
+++ b/Engine/Source/Developer/SourceControl/Private/SourceControlAssetDataCache.cpp
@@ -38,6 +38,12 @@ void FSourceControlAssetDataCache::Startup()
 		{
 			if (UPackage* Package = ChangedActor->GetPackage())
 			{
+				// Transient packages should not be saved, therefore we can skip the asset data cache update for them to mitigate the perf impact of this call 
+				if (Package->HasAnyFlags(RF_Transient))
+				{
+					return;
+				}
+
 				FString Filename = USourceControlHelpers::PackageFilename(Package);
 				ISourceControlModule::Get().GetAssetDataCache().ClearAssetData(Filename);
 			}
