Skein / SnapshotHistory: Icon Update.

diff --git a/Engine/Source/Developer/SourceControl/Private/RevisionControlStyle/RevisionControlStyle.cpp b/Engine/Source/Developer/SourceControl/Private/RevisionControlStyle/RevisionControlStyle.cpp
--- a/Engine/Source/Developer/SourceControl/Private/RevisionControlStyle/RevisionControlStyle.cpp
+++ b/Engine/Source/Developer/SourceControl/Private/RevisionControlStyle/RevisionControlStyle.cpp
@@ -80,6 +80,10 @@ FDefaultRevisionControlStyle::FDefaultRevisionControlStyle() : FSlateStyleSet(St
 	StatusCheckedOutByOtherUserColor = FLinearColor::FromSRGBColor(FColor::FromHex("#EF3535"));
 	StatusNotAtHeadRevisionColor = FLinearColor::FromSRGBColor(FColor::FromHex("#E1FF3D"));
 
+	SnapshotHistoryAdded = FLinearColor::FromSRGBColor(FColor::FromHex("#00CAA5"));
+	SnapshotHistoryModified = FLinearColor::FromSRGBColor(FColor::FromHex("#F0AD4E"));
+	SnapshotHistoryRemoved = FLinearColor::FromSRGBColor(FColor::FromHex("#CD3642"));
+
 	// Status icons
 	Set("RevisionControl.Icon", new IMAGE_BRUSH_SVG("Starship/SourceControl/Status/RevisionControl", CoreStyleConstants::Icon16x16));
 	Set("RevisionControl.Icon.ConnectedBadge", new IMAGE_BRUSH_SVG("Starship/SourceControl/Status/RevisionControlBadgeConnected", CoreStyleConstants::Icon16x16, FStyleColors::Success));
@@ -150,6 +154,10 @@ FDefaultRevisionControlStyle::FDefaultRevisionControlStyle() : FSlateStyleSet(St
 	Set("RevisionControl.ConflictResolution.OpenExternal", new IMAGE_BRUSH_SVG("Starship/SourceControl/RC_ConflictResolution_OpenExternal", CoreStyleConstants::Icon16x16));
 	Set("RevisionControl.ConflictResolution.Clear", new IMAGE_BRUSH_SVG("Starship/SourceControl/RC_ConflictResolution_Clear", CoreStyleConstants::Icon16x16));
 
+	Set("RevisionControl.SnapshotHistory.Added", new IMAGE_BRUSH_SVG("Starship/SourceControl/RC_Added", CoreStyleConstants::Icon14x14, SnapshotHistoryAdded));
+	Set("RevisionControl.SnapshotHistory.Modified", new IMAGE_BRUSH_SVG("Starship/SourceControl/RC_Modified", CoreStyleConstants::Icon14x14, SnapshotHistoryModified));
+	Set("RevisionControl.SnapshotHistory.Removed", new IMAGE_BRUSH_SVG("Starship/SourceControl/RC_Removed", CoreStyleConstants::Icon14x14, SnapshotHistoryRemoved));
+
 	// Revision Control States
 	Set("RevisionControl.VerticalLine", new IMAGE_BRUSH_SVG("Starship/SourceControl/RC_VerticalLine", CoreStyleConstants::Icon26x26));
 	Set("RevisionControl.VerticalLineStart", new IMAGE_BRUSH_SVG("Starship/SourceControl/RC_VerticalLineStart", CoreStyleConstants::Icon26x26));
@@ -163,9 +171,6 @@ FDefaultRevisionControlStyle::FDefaultRevisionControlStyle() : FSlateStyleSet(St
 
 	Set("RevisionControl.File", new IMAGE_BRUSH_SVG("Starship/SourceControl/RC_File", CoreStyleConstants::Icon6x8));
 	Set("RevisionControl.DiskSize", new IMAGE_BRUSH_SVG("Starship/SourceControl/RC_DiskSize", CoreStyleConstants::Icon8x8));
-	Set("RevisionControl.State.Added", new IMAGE_BRUSH_SVG("Starship/SourceControl/RC_Added", CoreStyleConstants::Icon12x12));
-	Set("RevisionControl.State.Modified", new IMAGE_BRUSH_SVG("Starship/SourceControl/RC_Modified", CoreStyleConstants::Icon12x12));
-	Set("RevisionControl.State.Removed", new IMAGE_BRUSH_SVG("Starship/SourceControl/RC_Removed", CoreStyleConstants::Icon12x12));
 }
 
 FDefaultRevisionControlStyle::~FDefaultRevisionControlStyle()
diff --git a/Engine/Source/Developer/SourceControl/Public/RevisionControlStyle/RevisionControlStyle.h b/Engine/Source/Developer/SourceControl/Public/RevisionControlStyle/RevisionControlStyle.h
--- a/Engine/Source/Developer/SourceControl/Public/RevisionControlStyle/RevisionControlStyle.h
+++ b/Engine/Source/Developer/SourceControl/Public/RevisionControlStyle/RevisionControlStyle.h
@@ -61,6 +61,11 @@ protected:
 	FLinearColor StatusCheckedOutColor;
 	FLinearColor StatusCheckedOutByOtherUserColor;
 	FLinearColor StatusNotAtHeadRevisionColor;
+
+	/** The specific colors we use for all the "Snapshot History" state icons */
+	FLinearColor SnapshotHistoryAdded;
+	FLinearColor SnapshotHistoryModified;
+	FLinearColor SnapshotHistoryRemoved;
 	
 private:
 	static FName StyleName;
