Skein / SnapshotHistory: Added snapshot content details tooltip on hovering over the number of affected files in a snapshot.

diff --git a/Engine/Source/Developer/SourceControl/Private/RevisionControlStyle/RevisionControlStyle.cpp b/Engine/Source/Developer/SourceControl/Private/RevisionControlStyle/RevisionControlStyle.cpp
--- a/Engine/Source/Developer/SourceControl/Private/RevisionControlStyle/RevisionControlStyle.cpp
+++ b/Engine/Source/Developer/SourceControl/Private/RevisionControlStyle/RevisionControlStyle.cpp
@@ -158,6 +158,10 @@ FDefaultRevisionControlStyle::FDefaultRevisionControlStyle() : FSlateStyleSet(St
 	Set("RevisionControl.SnapshotHistory.Modified", new IMAGE_BRUSH_SVG("Starship/SourceControl/RC_Modified", CoreStyleConstants::Icon14x14, SnapshotHistoryModified));
 	Set("RevisionControl.SnapshotHistory.Removed", new IMAGE_BRUSH_SVG("Starship/SourceControl/RC_Removed", CoreStyleConstants::Icon14x14, SnapshotHistoryRemoved));
 
+	Set("RevisionControl.SnapshotHistory.ToolTip.Added", new IMAGE_BRUSH_SVG("Starship/SourceControl/RC_Added", CoreStyleConstants::Icon14x14));
+	Set("RevisionControl.SnapshotHistory.ToolTip.Modified", new IMAGE_BRUSH_SVG("Starship/SourceControl/RC_Modified", CoreStyleConstants::Icon14x14));
+	Set("RevisionControl.SnapshotHistory.ToolTip.Removed", new IMAGE_BRUSH_SVG("Starship/SourceControl/RC_Removed", CoreStyleConstants::Icon14x14));
+
 	// Revision Control States
 	Set("RevisionControl.VerticalLine", new IMAGE_BRUSH_SVG("Starship/SourceControl/RC_VerticalLine", CoreStyleConstants::Icon26x26));
 	Set("RevisionControl.VerticalLineStart", new IMAGE_BRUSH_SVG("Starship/SourceControl/RC_VerticalLineStart", CoreStyleConstants::Icon26x26));
