[SourceControl] Delete of file that is not at latest fails silently and removes the entry from content browser

diff --git a/Engine/Source/Developer/SourceControl/Private/SourceControlModule.cpp b/Engine/Source/Developer/SourceControl/Private/SourceControlModule.cpp
--- a/Engine/Source/Developer/SourceControl/Private/SourceControlModule.cpp
+++ b/Engine/Source/Developer/SourceControl/Private/SourceControlModule.cpp
@@ -23,6 +23,7 @@
 #endif
 
 #if WITH_EDITOR
+	#include "Editor.h"
 	#include "EngineAnalytics.h"
 	#include "Interfaces/IAnalyticsProvider.h"
 #endif
@@ -68,6 +69,10 @@ void FSourceControlModule::StartupModule()
 
 	AssetDataCache.Startup();
 
+#if WITH_EDITOR
+	FEditorDelegates::OnAssetsCanDelete.AddRaw(this, &FSourceControlModule::OnAssetsCanDelete);
+#endif
+	
 	SourceControlFileStatusMonitor = MakeShared<FSourceControlFileStatusMonitor>();
 }
 
@@ -93,6 +98,10 @@ void FSourceControlModule::ShutdownModule()
 	IModularFeatures::Get().OnModularFeatureRegistered().RemoveAll(this);
 	IModularFeatures::Get().OnModularFeatureUnregistered().RemoveAll(this);
 
+#if WITH_EDITOR
+	FEditorDelegates::OnAssetsCanDelete.RemoveAll(this);
+#endif
+
 	SourceControlFileStatusMonitor.Reset();
 }
 
@@ -468,6 +477,51 @@ void FSourceControlModule::HandleModularFeatureUnregistered(const FName& Type, I
 	}
 }
 
+#if WITH_EDITOR
+void FSourceControlModule::OnAssetsCanDelete(const TArray<UObject*>& InObjects, FCanDeleteAssetResult& OutCanDeleteAssetResult) const
+{
+	ISourceControlProvider& SourceControlProvider = GetProvider();
+	if (IsEnabled() && SourceControlProvider.IsAvailable())
+	{
+		TArray<FString> PackagesNames;
+		PackagesNames.Reserve(InObjects.Num());
+		for (UObject* Object : InObjects)
+		{
+			if (Object)
+			{
+				if (UPackage* ObjectPackage = Object->GetPackage())
+				{
+					FString PackageFilename;
+					if (FPackageName::DoesPackageExist(ObjectPackage->GetName(), &PackageFilename))
+					{
+						PackagesNames.Add(PackageFilename);
+					}
+				}
+			}
+		}
+
+		TArray<FSourceControlStateRef> PackageSCCStates;
+		if (PackagesNames.Num())
+		{
+			SourceControlProvider.GetState(PackagesNames, PackageSCCStates, EStateCacheUsage::ForceUpdate);
+		}
+
+		for (const FSourceControlStateRef& SCCState : PackageSCCStates)
+		{
+			if (SCCState->IsSourceControlled())
+			{
+				if (!SCCState->CanDelete())
+				{
+					OutCanDeleteAssetResult.Set(false);
+					UE_LOG(LogSourceControl, Warning, TEXT("Asset(s) can't be deleted because at least one is currently checked out by another user and/or not at latest."));
+					return;
+				}
+			}
+		}
+	}
+}
+#endif
+
 bool FSourceControlModule::GetUseGlobalSettings() const
 {
 	return SourceControlSettings.GetUseGlobalSettings();
diff --git a/Engine/Source/Developer/SourceControl/Private/SourceControlModule.h b/Engine/Source/Developer/SourceControl/Private/SourceControlModule.h
--- a/Engine/Source/Developer/SourceControl/Private/SourceControlModule.h
+++ b/Engine/Source/Developer/SourceControl/Private/SourceControlModule.h
@@ -17,6 +17,8 @@ class FSourceControlFileStatusMonitor;
 class SSourceControlLogin;
 class SWindow;
 
+struct FCanDeleteAssetResult;
+
 class FSourceControlModule : public ISourceControlModule
 {
 public:
@@ -122,6 +124,10 @@ private:
 	/** Delegate handling when source control features are unregistered */
 	void HandleModularFeatureUnregistered(const FName& Type, IModularFeature* ModularFeature);
 
+#if WITH_EDITOR
+	/** Delegate called when trying to delete asset(s) to verify that it is allowed by the SourceControl */
+	void OnAssetsCanDelete(const TArray<UObject*>& InObjects, FCanDeleteAssetResult& OutCanDeleteAssetResult) const;
+#endif
 private:
 	/** The settings object */
 	FSourceControlSettings SourceControlSettings;
