RevisionControlStyle: Fixed discrepancy between '+' sign color in Snapshot History assets row and SceneOutliner for 'added' files.

diff --git a/Engine/Source/Developer/SourceControl/Private/RevisionControlStyle/RevisionControlStyle.cpp b/Engine/Source/Developer/SourceControl/Private/RevisionControlStyle/RevisionControlStyle.cpp
--- a/Engine/Source/Developer/SourceControl/Private/RevisionControlStyle/RevisionControlStyle.cpp
+++ b/Engine/Source/Developer/SourceControl/Private/RevisionControlStyle/RevisionControlStyle.cpp
@@ -80,7 +80,7 @@ FDefaultRevisionControlStyle::FDefaultRevisionControlStyle() : FSlateStyleSet(St
 	StatusCheckedOutByOtherUserColor = FLinearColor::FromSRGBColor(FColor::FromHex("#EF3535"));
 	StatusNotAtHeadRevisionColor = FLinearColor::FromSRGBColor(FColor::FromHex("#E1FF3D"));
 
-	SnapshotHistoryAdded = FLinearColor::FromSRGBColor(FColor::FromHex("#00CAA5"));
+	SnapshotHistoryAdded = FStyleColors::AccentBlue;
 	SnapshotHistoryModified = FLinearColor::FromSRGBColor(FColor::FromHex("#F0AD4E"));
 	SnapshotHistoryRemoved = FLinearColor::FromSRGBColor(FColor::FromHex("#CD3642"));
 
diff --git a/Engine/Source/Developer/SourceControl/Public/RevisionControlStyle/RevisionControlStyle.h b/Engine/Source/Developer/SourceControl/Public/RevisionControlStyle/RevisionControlStyle.h
--- a/Engine/Source/Developer/SourceControl/Public/RevisionControlStyle/RevisionControlStyle.h
+++ b/Engine/Source/Developer/SourceControl/Public/RevisionControlStyle/RevisionControlStyle.h
@@ -55,17 +55,17 @@ public:
 protected:
 
 	/** The specific color we use for all the "Branched" icons */
-	FLinearColor BranchedColor;
+	FSlateColor BranchedColor;
 
 	/** The specific colors we use for all the "Status" icons */
-	FLinearColor StatusCheckedOutColor;
-	FLinearColor StatusCheckedOutByOtherUserColor;
-	FLinearColor StatusNotAtHeadRevisionColor;
+	FSlateColor StatusCheckedOutColor;
+	FSlateColor StatusCheckedOutByOtherUserColor;
+	FSlateColor StatusNotAtHeadRevisionColor;
 
 	/** The specific colors we use for all the "Snapshot History" state icons */
-	FLinearColor SnapshotHistoryAdded;
-	FLinearColor SnapshotHistoryModified;
-	FLinearColor SnapshotHistoryRemoved;
+	FSlateColor SnapshotHistoryAdded;
+	FSlateColor SnapshotHistoryModified;
+	FSlateColor SnapshotHistoryRemoved;
 	
 private:
 	static FName StyleName;
