#UE Fix CheckOutOrAddFiles so it works in the common case of a file just being in the depot and not checked out

diff --git a/Engine/Source/Developer/SourceControl/Private/SourceControlHelpers.cpp b/Engine/Source/Developer/SourceControl/Private/SourceControlHelpers.cpp
--- a/Engine/Source/Developer/SourceControl/Private/SourceControlHelpers.cpp
+++ b/Engine/Source/Developer/SourceControl/Private/SourceControlHelpers.cpp
@@ -599,9 +599,9 @@ bool USourceControlHelpers::CheckOutOrAddFiles(const TArray<FString>& InFiles, b
 
 		// Less error checking and info is made for multiple files than the single file version.
 		// This multi-file version could be made similarly more sophisticated.
-		if (!SCState->IsCheckedOut())
+		if (!SCState->IsCheckedOut() && !SCState->IsAdded())
 		{
-			if (!SCState->IsAdded())
+			if (!SCState->IsSourceControlled())
 			{
 				if (SCState->CanAdd())
 				{
