The 'Revision Control Login' dialog will now immedately save revision control settings to disk when either 'accept' or 'disable' buttons are selected.

- Previously we only saved the settings when the editor shutdown or an unrelated full flush of the config file system was triggered elsewhere. People often found this confusing as they would set up their editor with a revision control setting only to find that it was no longer set if the editor process crashed or was terminated by a code debugger etc./
- Now we immediately flush our settings to disk once the user has accepted the new settings or opted to disable revision control entirely.
- Fixed SOURCE_CONTROL_WITH_SLATE to cover the entire .cpp and .h file, none of the code that was outside of these pre-existing scopes is required when SOURCE_CONTROL_WITH_SLATE == 0.

diff --git a/Engine/Source/Developer/SourceControl/Private/SSourceControlLogin.cpp b/Engine/Source/Developer/SourceControl/Private/SSourceControlLogin.cpp
--- a/Engine/Source/Developer/SourceControl/Private/SSourceControlLogin.cpp
+++ b/Engine/Source/Developer/SourceControl/Private/SSourceControlLogin.cpp
@@ -1,26 +1,25 @@
 // Copyright Epic Games, Inc. All Rights Reserved.
 
 #include "SSourceControlLogin.h"
-#include "SourceControlModule.h"
-#include "SourceControlOperations.h"
 
 #if SOURCE_CONTROL_WITH_SLATE
-#include "Widgets/SBoxPanel.h"
+
 #include "Framework/Application/SlateApplication.h"
-#include "Widgets/Text/STextBlock.h"
-#include "Widgets/Layout/SBox.h"
-#include "Widgets/Layout/SUniformGridPanel.h"
-#include "Widgets/Input/SButton.h"
+#include "Framework/Notifications/NotificationManager.h"
+#include "Logging/MessageLog.h"
+#include "Misc/ConfigCacheIni.h"
+#include "SSourceControlPicker.h"
+#include "SourceControlHelpers.h"
+#include "SourceControlModule.h"
+#include "SourceControlOperations.h"
 #include "Widgets/Images/SThrobber.h"
+#include "Widgets/Input/SButton.h"
+#include "Widgets/Layout/SBox.h"
 #include "Widgets/Layout/SExpandableArea.h"
-#include "Framework/Notifications/NotificationManager.h"
+#include "Widgets/Layout/SUniformGridPanel.h"
 #include "Widgets/Notifications/SNotificationList.h"
-#endif
-
-#include "SSourceControlPicker.h"
-#include "Logging/MessageLog.h"
-
-#if SOURCE_CONTROL_WITH_SLATE
+#include "Widgets/SBoxPanel.h"
+#include "Widgets/Text/STextBlock.h"
 
 #if WITH_UNREAL_DEVELOPER_TOOLS
 	#include "MessageLogModule.h"
@@ -181,7 +180,12 @@ FReply SSourceControlLogin::OnAcceptSettings()
 	}
 
 	FSourceControlModule& SourceControlModule = FSourceControlModule::Get();
-	if(SourceControlModule.GetProvider().Login(FString(), EConcurrency::Asynchronous, FSourceControlOperationComplete::CreateSP(this, &SSourceControlLogin::SourceControlOperationComplete)) == ECommandResult::Type::Failed)
+	ECommandResult::Type Result = SourceControlModule.GetProvider().Login(FString(), EConcurrency::Asynchronous, FSourceControlOperationComplete::CreateSP(this, &SSourceControlLogin::SourceControlOperationComplete));
+	if (Result == ECommandResult::Succeeded)
+	{
+		SaveSettingsToDisk();
+	}
+	else if(Result == ECommandResult::Failed)
 	{
 		DisplayConnectionError();
 		ConnectionState = ELoginConnectionState::Disconnected;
@@ -214,6 +218,9 @@ FReply SSourceControlLogin::OnDisableSourceControl()
 		ParentWindowPtr.Pin()->RequestDestroyWindow();
 	}
 	SourceControlLoginClosed.ExecuteIfBound(false);
+
+	SaveSettingsToDisk();
+
 	return FReply::Handled();
 }
 
@@ -305,6 +312,12 @@ EActiveTimerReturnType SSourceControlLogin::TickSourceControlModule( double InCu
 	return EActiveTimerReturnType::Continue;
 }
 
+void SSourceControlLogin::SaveSettingsToDisk() const
+{
+	check(::IsInGameThread());
+	GConfig->Flush(false, USourceControlHelpers::GetSettingsIni());
+}
+
 #undef LOCTEXT_NAMESPACE
 
 #endif // SOURCE_CONTROL_WITH_SLATE
diff --git a/Engine/Source/Developer/SourceControl/Private/SSourceControlLogin.h b/Engine/Source/Developer/SourceControl/Private/SSourceControlLogin.h
--- a/Engine/Source/Developer/SourceControl/Private/SSourceControlLogin.h
+++ b/Engine/Source/Developer/SourceControl/Private/SSourceControlLogin.h
@@ -2,21 +2,20 @@
 
 #pragma once
 
-#include "CoreMinimal.h"
+#if SOURCE_CONTROL_WITH_SLATE
+
 #include "ISourceControlModule.h"
 #include "ISourceControlProvider.h"
-#if SOURCE_CONTROL_WITH_SLATE
+#include "Input/Reply.h"
 #include "Layout/Visibility.h"
 #include "Widgets/DeclarativeSyntaxSupport.h"
-#include "Input/Reply.h"
-#include "Widgets/SCompoundWidget.h"
 #include "Widgets/Layout/SBorder.h"
-#endif
+#include "Widgets/SCompoundWidget.h"
 
 class FActiveTimerHandle;
 class IDetailsView;
-class SWindow;
 class SBox;
+class SWindow;
 
 namespace ELoginConnectionState
 {
@@ -28,9 +27,6 @@ namespace ELoginConnectionState
 	};
 }
 
-#if SOURCE_CONTROL_WITH_SLATE
-
-
 class SSourceControlLogin : public SCompoundWidget
 {
 public:
@@ -92,6 +88,9 @@ private:
 	/** Requests confirmation of disabling source control */
 	bool ConfirmDisableSourceControl() const;
 
+	/** Flushes all current revision control settings to disk */
+	void SaveSettingsToDisk() const;
+
 private:
 	/** The frequency at which to tick to scc module when inside a modal window */
 	static const float RefreshFrequency;
