Add a feature to SourceControlModule that allows operations derived from ISourceControlOperation to request that errors not be immediately logged.

- Some systems that use source control might want to retry operations that fail multiple times to see if they can eventually be completed. Under these conditions it is not always desirable to log the errors to stdout as that might cause the hosting systems (such as a build environment) to fail an entire job when it sees that an error has been encountered. Instead the calling code would like to be able to tell the source control module (and it's implementation) to not log errors and then it can take the decision to log or not. It is expected that info/warning/errors continue to be returned to the caller via FSourceControlOperationBase::ResultInfo.
- This change only adds the infrastructure for this feature, it will be up to future submits to add support for this to each source control implementation.
- At the base level ISourceControlOperation::ShouldLogErrors can be overridden by derived classes and return false to prevent errors from being immediately logged.
- For now we implement this at FSourceControlOperationBase level via a new enum bitfield which at the moment only supports disabling error logging.
-- At some point in the future this system should be expanded to be able to selectively enable/disable info, warnings and errors independently which is why this control has been implemented as a bitfield.
- FDownloadFile already had support for something similar via its EVerbosity flag. For now we support using this older flag to set the newer system by calling ::SetEnableErrorLogging.
-- This enum will be deprecated/removed in a future submit.

diff --git a/Engine/Source/Developer/SourceControl/Private/SourceControlOperations.cpp b/Engine/Source/Developer/SourceControl/Private/SourceControlOperations.cpp
--- a/Engine/Source/Developer/SourceControl/Private/SourceControlOperations.cpp
+++ b/Engine/Source/Developer/SourceControl/Private/SourceControlOperations.cpp
@@ -45,8 +45,7 @@ bool FUpdateStatus::IsDirectoryPath(const FString& Path) const
 }
 
 FDownloadFile::FDownloadFile(FStringView InTargetDirectory, EVerbosity InVerbosity)
-	: Verbosity(InVerbosity)
-	, TargetDirectory(InTargetDirectory)
+	: TargetDirectory(InTargetDirectory)
 {
 	FPaths::NormalizeDirectoryName(TargetDirectory);
 
diff --git a/Engine/Source/Developer/SourceControl/Public/ISourceControlOperation.h b/Engine/Source/Developer/SourceControl/Public/ISourceControlOperation.h
--- a/Engine/Source/Developer/SourceControl/Public/ISourceControlOperation.h
+++ b/Engine/Source/Developer/SourceControl/Public/ISourceControlOperation.h
@@ -67,6 +67,14 @@ public:
 		return false;
 	}
 
+	/**
+	 * Returns true when the operation is allowed to log errors immediately when they are encountered.
+	 */
+	virtual bool ShouldLogErrors() const
+	{
+		return true;
+	}
+
 	/** Factory method for easier operation creation */
 	template<typename Type, typename... TArgs>
 	static TSharedRef<Type, ESPMode::ThreadSafe> Create(TArgs&&... Args)
diff --git a/Engine/Source/Developer/SourceControl/Public/SourceControlOperationBase.h b/Engine/Source/Developer/SourceControl/Public/SourceControlOperationBase.h
--- a/Engine/Source/Developer/SourceControl/Public/SourceControlOperationBase.h
+++ b/Engine/Source/Developer/SourceControl/Public/SourceControlOperationBase.h
@@ -8,6 +8,18 @@
 class FSourceControlOperationBase : public ISourceControlOperation
 {
 public:
+	enum class EFlags
+	{
+		/* No specialization applied */
+		None = 0,
+		/* When set the operation should not log any error but should continue to store them in the ResultInfo member */
+		DisableErrorLogging = 1 << 0
+	};
+	FRIEND_ENUM_CLASS_FLAGS(EFlags)
+
+	FSourceControlOperationBase() = default;
+	virtual ~FSourceControlOperationBase() = default;
+
 	/** Retrieve any info or error messages that may have accumulated during the operation. */
 	virtual const FSourceControlResultInfo& GetResultInfo() const override
 	{
@@ -41,5 +53,27 @@ public:
 		ResultInfo.Append(InResultInfo);
 	}
 
+	void SetEnableErrorLogging(bool bEnabled)
+	{
+		if (bEnabled)
+		{
+			EnumRemoveFlags(Flags, EFlags::DisableErrorLogging);
+		}
+		else
+		{
+			EnumAddFlags(Flags, EFlags::DisableErrorLogging);
+		}
+	}
+
+	virtual bool ShouldLogErrors() const override
+	{
+		return !EnumHasAllFlags(Flags, EFlags::DisableErrorLogging);
+	}
+
 	FSourceControlResultInfo ResultInfo;
+
+protected:
+	EFlags Flags = EFlags::None;
 };
+
+ENUM_CLASS_FLAGS(FSourceControlOperationBase::EFlags);
diff --git a/Engine/Source/Developer/SourceControl/Public/SourceControlOperations.h b/Engine/Source/Developer/SourceControl/Public/SourceControlOperations.h
--- a/Engine/Source/Developer/SourceControl/Public/SourceControlOperations.h
+++ b/Engine/Source/Developer/SourceControl/Public/SourceControlOperations.h
@@ -1167,7 +1167,12 @@ public:
 	 */
 	FDownloadFile(EVerbosity InVerbosity)
 		: Verbosity(InVerbosity)
-	{}
+	{
+		if (InVerbosity == EVerbosity::None)
+		{
+			SetEnableErrorLogging(false);
+		}
+	}
 
 	// ISourceControlOperation interface
 	virtual FName GetName() const override
@@ -1210,7 +1215,7 @@ public:
 	/** Return true if the command should log its operations, otherwise false */
 	bool ShouldLogToStdOutput() const
 	{
-		return Verbosity == EVerbosity::Full;
+		return !EnumHasAllFlags(Flags, EFlags::DisableErrorLogging);
 	}
 
 	/** 
