Deprecate FDownloadFile::EVerbosity as it's functionality has been replaced by FSourceControlOperationBase::SetEnableErrorLogging.

diff --git a/Engine/Source/Developer/SourceControl/Private/SourceControlOperations.cpp b/Engine/Source/Developer/SourceControl/Private/SourceControlOperations.cpp
--- a/Engine/Source/Developer/SourceControl/Private/SourceControlOperations.cpp
+++ b/Engine/Source/Developer/SourceControl/Private/SourceControlOperations.cpp
@@ -44,7 +44,18 @@ bool FUpdateStatus::IsDirectoryPath(const FString& Path) const
 	}
 }
 
+PRAGMA_DISABLE_DEPRECATION_WARNINGS
 FDownloadFile::FDownloadFile(FStringView InTargetDirectory, EVerbosity InVerbosity)
+	: FDownloadFile(InTargetDirectory)
+{
+	if (InVerbosity == EVerbosity::None)
+	{
+		SetEnableErrorLogging(false);
+	}
+}
+PRAGMA_ENABLE_DEPRECATION_WARNINGS
+
+FDownloadFile::FDownloadFile(FStringView InTargetDirectory)
 	: TargetDirectory(InTargetDirectory)
 {
 	FPaths::NormalizeDirectoryName(TargetDirectory);
diff --git a/Engine/Source/Developer/SourceControl/Public/SourceControlOperations.h b/Engine/Source/Developer/SourceControl/Public/SourceControlOperations.h
--- a/Engine/Source/Developer/SourceControl/Public/SourceControlOperations.h
+++ b/Engine/Source/Developer/SourceControl/Public/SourceControlOperations.h
@@ -1146,7 +1146,7 @@ public:
 class FDownloadFile : public FSourceControlOperationBase
 {
 public:
-	enum class EVerbosity
+	enum class UE_DEPRECATED(5.6, "Enum has been removed, call FSourceControlOperationBase::SetEnableErrorLogging directly instead") EVerbosity
 	{
 		/** No logging when the command is run. */
 		None,
@@ -1156,23 +1156,19 @@ public:
 
 	/** 
 	 * This version of FDownloadFile will download the files and keep them in memory,
-	 * which can then be accessed by calling ::GetFileData. Full logging verbosity
-	 * will be used.
+	 * which can then be accessed by calling ::GetFileData.
 	 */
 	FDownloadFile() = default;
 
-	/**
-	 * This version of FDownloadFile will download the files and keep them in memory,
-	 * which can then be accessed by calling ::GetFileData.
-	 */
+PRAGMA_DISABLE_DEPRECATION_WARNINGS
 	FDownloadFile(EVerbosity InVerbosity)
-		: Verbosity(InVerbosity)
 	{
 		if (InVerbosity == EVerbosity::None)
 		{
 			SetEnableErrorLogging(false);
 		}
 	}
+PRAGMA_ENABLE_DEPRECATION_WARNINGS
 
 	// ISourceControlOperation interface
 	virtual FName GetName() const override
@@ -1192,8 +1188,13 @@ public:
 		return true;
 	}
 	
-	/** This version of FDownloadFile will download the files to the given target directory */
+PRAGMA_DISABLE_DEPRECATION_WARNINGS
+	UE_DEPRECATED(5.6, "Enum has been removed, use the overload without InVerbosity instead")
 	SOURCECONTROL_API FDownloadFile(FStringView InTargetDirectory, EVerbosity InVerbosity);
+PRAGMA_ENABLE_DEPRECATION_WARNINGS
+
+	/** This version of FDownloadFile will download the files to the given target directory */
+	SOURCECONTROL_API FDownloadFile(FStringView InTargetDirectory);
 
 	/**
 	 * Returns the directory that the files will (or have) been downloaded to. This path
@@ -1235,8 +1236,6 @@ public:
 
 private:
 
-	EVerbosity Verbosity = EVerbosity::Full;
-
 	FString TargetDirectory;
 	TMap<FString, FSharedBuffer> FileDataMap;
 };
diff --git a/Engine/Source/Developer/Virtualization/Private/VirtualizationSourceControlBackend.cpp b/Engine/Source/Developer/Virtualization/Private/VirtualizationSourceControlBackend.cpp
--- a/Engine/Source/Developer/Virtualization/Private/VirtualizationSourceControlBackend.cpp
+++ b/Engine/Source/Developer/Virtualization/Private/VirtualizationSourceControlBackend.cpp
@@ -28,8 +28,6 @@
 #include "VirtualizationSourceControlUtilities.h"
 #include "VirtualizationUtilities.h"
 
-#include "HAL/PlatformFileManager.h"
-
 // When the SourceControl module (or at least the perforce source control module) is thread safe we
 // can enable this and stop using the hacky work around 'TryToDownloadFileFromBackgroundThread'
 #define IS_SOURCE_CONTROL_THREAD_SAFE 0
@@ -568,7 +566,8 @@ bool FSourceControlBackend::PullData(TArrayView<FPullRequest> Requests, EPullFla
 			FPlatformProcess::SleepNoStats(static_cast<float>(RetryWaitTimeMS) * 0.001f);
 		}
 
-		TSharedRef<FDownloadFile> DownloadCommand = ISourceControlOperation::Create<FDownloadFile>(FDownloadFile::EVerbosity::None);
+		TSharedRef<FDownloadFile> DownloadCommand = ISourceControlOperation::Create<FDownloadFile>();
+		DownloadCommand->SetEnableErrorLogging(false);
 
 #if IS_SOURCE_CONTROL_THREAD_SAFE
 		const bool bOperationSuccess = SCCProvider->Execute(DownloadCommand, DepotPaths, EConcurrency::Synchronous) == ECommandResult::Succeeded;
diff --git a/Engine/Source/Editor/UnrealEd/Private/MergeUtils.cpp b/Engine/Source/Editor/UnrealEd/Private/MergeUtils.cpp
--- a/Engine/Source/Editor/UnrealEd/Private/MergeUtils.cpp
+++ b/Engine/Source/Editor/UnrealEd/Private/MergeUtils.cpp
@@ -78,7 +78,7 @@ UPackage* MergeUtils::LoadPackageForMerge(const FString& SCFile, const FString&
 FString MergeUtils::LoadSCFileForMerge(const FString& SCFile, const FString& Revision)
 {
 	const FString FileWithRevision = SCFile + TEXT("#") + Revision;
-	const TSharedRef<FDownloadFile, ESPMode::ThreadSafe> DownloadFileOperation = ISourceControlOperation::Create<FDownloadFile>(FPaths::DiffDir(), FDownloadFile::EVerbosity::Full);
+	const TSharedRef<FDownloadFile, ESPMode::ThreadSafe> DownloadFileOperation = ISourceControlOperation::Create<FDownloadFile>(FPaths::DiffDir());
 	ISourceControlModule::Get().GetProvider().Execute(DownloadFileOperation, FileWithRevision, EConcurrency::Synchronous);
 	const FString DownloadPath = FPaths::ConvertRelativePathToFull(FPaths::DiffDir() / FPaths::GetCleanFilename(FileWithRevision));
 
