Added FScopedDisableSourceControl

This utilizes the existing bTemporarilyDisabled flag (now a counter) to temporarily cause ISourceControlModule::IsEnabled to return false, without actually disabling the current source control provider.

This is more reliable than disabling the current provider for things like temporarily disabling source control when saving writable files, as disabling the provider causes other systems to interact with that change and potentially do expensive refreshes. It also keeps things like Perforce connected to their server, and avoids losing any source control settings.

Note: ISourceControlModule::GetProvider still returns the current source control provider even when temporarily disabled, so code must also check ISourceControlModule::IsEnabled before using it.

diff --git a/Engine/Source/Developer/SourceControl/Private/SourceControlModule.cpp b/Engine/Source/Developer/SourceControl/Private/SourceControlModule.cpp
--- a/Engine/Source/Developer/SourceControl/Private/SourceControlModule.cpp
+++ b/Engine/Source/Developer/SourceControl/Private/SourceControlModule.cpp
@@ -43,9 +43,28 @@ namespace SourceControlConstants
 	const FTimespan StateRefreshInterval = FTimespan::FromMinutes(5.0);
 }
 
+FScopedDisableSourceControl::FScopedDisableSourceControl(const bool InDisable)
+	: bDisabled(InDisable)
+{
+	checkf(IsInGameThread(), TEXT("FScopedDisableSourceControl should only be used from the game thread!"));
+	if (bDisabled)
+	{
+		++FSourceControlModule::Get().TemporarilyDisabledCount;
+	}
+}
+
+FScopedDisableSourceControl::~FScopedDisableSourceControl()
+{
+	checkf(IsInGameThread(), TEXT("FScopedDisableSourceControl should only be used from the game thread!"));
+	if (bDisabled)
+	{
+		--FSourceControlModule::Get().TemporarilyDisabledCount;
+	}
+}
+
 FSourceControlModule::FSourceControlModule()
 	: CurrentSourceControlProvider(NULL)
-	, bTemporarilyDisabled(false)
+	, TemporarilyDisabledCount(0)
 {
 }
 
@@ -151,7 +170,7 @@ void FSourceControlModule::ShowLoginDialog(const FSourceControlLoginClosed& InOn
 		}
 
 		// temporarily disable access to source control features
-		bTemporarilyDisabled = true;
+		++TemporarilyDisabledCount;
 
 		// Create the window
 		SourceControlLoginWindowPtr = SNew(SWindow)
@@ -204,7 +223,7 @@ void FSourceControlModule::OnSourceControlDialogClosed(const TSharedRef<SWindow>
 {
 	SourceControlLoginWindowPtr = NULL;
 	SourceControlLoginPtr = NULL;
-	bTemporarilyDisabled = false;
+	--TemporarilyDisabledCount;
 
 #if WITH_EDITOR
 	FString NewProvider = CurrentSourceControlProvider->GetName().ToString();
@@ -271,7 +290,7 @@ void FSourceControlModule::Tick()
 		AssetDataCache.Tick();
 
 		// don't allow background status updates when temporarily disabled for login
-		if (!bTemporarilyDisabled)
+		if (TemporarilyDisabledCount == 0)
 		{
 			// check for any pending dispatches
 			if (PendingStatusUpdateFiles.Num() > 0)
@@ -349,7 +368,7 @@ void FSourceControlModule::QueueStatusUpdate(const FString& InFilename)
 
 bool FSourceControlModule::IsEnabled() const
 {
-	return !bTemporarilyDisabled && GetProvider().IsEnabled();
+	return TemporarilyDisabledCount == 0 && GetProvider().IsEnabled();
 }
 
 ISourceControlProvider& FSourceControlModule::GetProvider() const
diff --git a/Engine/Source/Developer/SourceControl/Private/SourceControlModule.h b/Engine/Source/Developer/SourceControl/Private/SourceControlModule.h
--- a/Engine/Source/Developer/SourceControl/Private/SourceControlModule.h
+++ b/Engine/Source/Developer/SourceControl/Private/SourceControlModule.h
@@ -151,7 +151,7 @@ private:
 	TArray<FString> PendingStatusUpdateFiles;
 
 	/** Flag to disable source control - used temporarily when login is in progress */
-	bool bTemporarilyDisabled;
+	uint8 TemporarilyDisabledCount;
 
 	/** Active Provider name to track source control provider changes */
 	FString ActiveProviderName;
@@ -179,4 +179,6 @@ private:
 
 	/** Delegate used to return information on custom projects under source control */
 	FSourceControlCustomProjectsDelegate CustomProjectsDelegate;
+
+	friend FScopedDisableSourceControl;
 };
diff --git a/Engine/Source/Developer/SourceControl/Public/ISourceControlModule.h b/Engine/Source/Developer/SourceControl/Public/ISourceControlModule.h
--- a/Engine/Source/Developer/SourceControl/Public/ISourceControlModule.h
+++ b/Engine/Source/Developer/SourceControl/Public/ISourceControlModule.h
@@ -98,6 +98,23 @@ namespace EOnLoginWindowStartup
 	};
 };
 
+/**
+ * Temporarily disable source control for the lifetime of this scoped object
+ * 
+ * @note While this is on scope, ISourceControlModule::IsEnabled() will return false, 
+ *       but ISourceControlModule::GetProvider() will still return the current provider
+ */
+struct FScopedDisableSourceControl
+{
+public:
+	SOURCECONTROL_API explicit FScopedDisableSourceControl(const bool InDisable = true);
+	SOURCECONTROL_API ~FScopedDisableSourceControl();
+
+	UE_NONCOPYABLE(FScopedDisableSourceControl);
+
+private:
+	bool bDisabled;
+};
 
 /**
  * Interface for talking to source control providers
