Added a callback for when the list of projects returned by ISourceControlModule::GetCustomProjects changes

ISourceControlModule::OnCustomProjectsChanged is called automatically for ISourceControlModule::RegisterCustomProjectsDelegate and ISourceControlModule::UnregisterCustomProjectsDelegate, but outside of that it is the responsibility of the code that registered FSourceControlCustomProjectsDelegate to handle change notifications

diff --git a/Engine/Source/Developer/SourceControl/Private/SourceControlModule.cpp b/Engine/Source/Developer/SourceControl/Private/SourceControlModule.cpp
--- a/Engine/Source/Developer/SourceControl/Private/SourceControlModule.cpp
+++ b/Engine/Source/Developer/SourceControl/Private/SourceControlModule.cpp
@@ -617,11 +617,13 @@ const FSourceControlFilesDeletedDelegate& FSourceControlModule::GetOnFilesDelete
 void FSourceControlModule::RegisterCustomProjectsDelegate(FSourceControlCustomProjectsDelegate InCustomProjectsDelegate)
 {
 	CustomProjectsDelegate = MoveTemp(InCustomProjectsDelegate);
+	OnCustomProjectsChanged().Broadcast();
 }
 
 void FSourceControlModule::UnregisterCustomProjectsDelegate()
 {
 	CustomProjectsDelegate = FSourceControlCustomProjectsDelegate();
+	OnCustomProjectsChanged().Broadcast();
 }
 
 TArray<FSourceControlProjectInfo> FSourceControlModule::GetCustomProjects() const
@@ -629,6 +631,11 @@ TArray<FSourceControlProjectInfo> FSourceControlModule::GetCustomProjects() cons
 	return CustomProjectsDelegate.IsBound() ? CustomProjectsDelegate.Execute() : TArray<FSourceControlProjectInfo>();
 }
 
+FSimpleMulticastDelegate& FSourceControlModule::OnCustomProjectsChanged()
+{
+	return CustomProjectsChangedDelegate;
+}
+
 FString FSourceControlModule::GetSourceControlProjectDir() const
 {
 	{
diff --git a/Engine/Source/Developer/SourceControl/Private/SourceControlModule.h b/Engine/Source/Developer/SourceControl/Private/SourceControlModule.h
--- a/Engine/Source/Developer/SourceControl/Private/SourceControlModule.h
+++ b/Engine/Source/Developer/SourceControl/Private/SourceControlModule.h
@@ -19,7 +19,7 @@ class SWindow;
 
 struct FCanDeleteAssetResult;
 
-class FSourceControlModule : public ISourceControlModule
+class FSourceControlModule final : public ISourceControlModule
 {
 public:
 	FSourceControlModule();
@@ -61,6 +61,7 @@ public:
 	virtual void RegisterCustomProjectsDelegate(FSourceControlCustomProjectsDelegate InCustomProjectsDelegate) override;
 	virtual void UnregisterCustomProjectsDelegate() override;
 	virtual TArray<FSourceControlProjectInfo> GetCustomProjects() const override;
+	virtual FSimpleMulticastDelegate& OnCustomProjectsChanged() override;
 	virtual FString GetSourceControlProjectDir() const override;
 
 	UE_DEPRECATED(5.5, "Use RegisterCustomProjectsDelegate instead.")
@@ -179,6 +180,7 @@ private:
 
 	/** Delegate used to return information on custom projects under source control */
 	FSourceControlCustomProjectsDelegate CustomProjectsDelegate;
+	FSimpleMulticastDelegate CustomProjectsChangedDelegate;
 
 	friend FScopedDisableSourceControl;
 };
diff --git a/Engine/Source/Developer/SourceControl/Public/ISourceControlModule.h b/Engine/Source/Developer/SourceControl/Public/ISourceControlModule.h
--- a/Engine/Source/Developer/SourceControl/Public/ISourceControlModule.h
+++ b/Engine/Source/Developer/SourceControl/Public/ISourceControlModule.h
@@ -28,6 +28,17 @@ struct FSourceControlProjectInfo
 		: ProjectDirectory(MoveTemp(InProjectDirectory))
 	{}
 
+	friend bool operator==(const FSourceControlProjectInfo& A, const FSourceControlProjectInfo& B)
+	{
+		return A.ProjectDirectory == B.ProjectDirectory
+			&& A.ContentDirectories == B.ContentDirectories;
+	}
+
+	friend bool operator!=(const FSourceControlProjectInfo& A, const FSourceControlProjectInfo& B)
+	{
+		return !(A == B);
+	}
+
 	/** Directory that contains project files */
 	FString ProjectDirectory;
 
@@ -298,6 +309,13 @@ public:
 	 */
 	virtual TArray<FSourceControlProjectInfo> GetCustomProjects() const = 0;
 
+	/**
+	 * Delegate called whenever the list of projects that GetCustomProjects returns has changed.
+	 * @note It is called automatically for RegisterCustomProjectsDelegate and UnregisterCustomProjectsDelegate, but outside of that 
+	         it is the responsibility of the code that registered FSourceControlCustomProjectsDelegate to handle change notifications.
+	 */
+	virtual FSimpleMulticastDelegate& OnCustomProjectsChanged() = 0;
+
 	/**
 	 * Return the main custom project under source control or the uproject dir (FPaths::ProjectDir) if there's no custom project
 	 */
@@ -327,4 +345,15 @@ public:
 		static FName SourceControlModule("SourceControl");
 		return FModuleManager::LoadModuleChecked<ISourceControlModule>(SourceControlModule);
 	}
+
+	/**
+	 * Gets a pointer to the source control module instance, if loaded.
+	 *
+	 * @return A pointer to the source control module, or null if not loaded.
+	 */
+	static inline ISourceControlModule* GetPtr()
+	{
+		static FName SourceControlModule("SourceControl");
+		return FModuleManager::GetModulePtr<ISourceControlModule>(SourceControlModule);
+	}
 };
