Reverting a source control added file with "Should Delete New Files on Revert" disabled will now track it as an uncontrolled add, rather than leave it in limbo

diff --git a/Engine/Source/Developer/SourceControl/Private/SourceControlHelpers.cpp b/Engine/Source/Developer/SourceControl/Private/SourceControlHelpers.cpp
--- a/Engine/Source/Developer/SourceControl/Private/SourceControlHelpers.cpp
+++ b/Engine/Source/Developer/SourceControl/Private/SourceControlHelpers.cpp
@@ -1323,14 +1323,37 @@ bool USourceControlHelpers::RevertAndReloadPackages(const TArray<FString>& InPac
 {
 	auto RevertOperation = [bRevertAll](const TArray<FString>& InPackagesToRevert) -> bool
 	{
-		ISourceControlProvider& SourceControlProvider = ISourceControlModule::Get().GetProvider();
-
-		auto OperationCompleteCallback = FSourceControlOperationComplete::CreateLambda([](const FSourceControlOperationRef& Operation, ECommandResult::Type InResult)
+		auto OperationCompleteCallback = FSourceControlOperationComplete::CreateLambda([InPackagesToRevert](const FSourceControlOperationRef& Operation, ECommandResult::Type InResult)
 		{
 			if (Operation->GetName() == TEXT("Revert"))
 			{
 				TSharedRef<FRevert> RevertOperation = StaticCastSharedRef<FRevert>(Operation);
 				ISourceControlModule::Get().GetOnFilesDeleted().Broadcast(RevertOperation->GetDeletedFiles());
+
+				// The revert may have caused added files to return to an uncontrolled state when not deleting newly added files
+				// Handle that here, and return any newly uncontrolled files to the default uncontrolled changelist so they're not left in limbo
+				if (GEditor && !RevertOperation->ShouldDeleteNewFiles())
+				{
+					TArray<FSourceControlStateRef> SCStates;
+					if (ISourceControlModule::Get().GetProvider().GetState(InPackagesToRevert, SCStates, EStateCacheUsage::ForceUpdate) == ECommandResult::Succeeded)
+					{
+						TArray<FString> NewlyUncontrolledFiles;
+						for (const FSourceControlStateRef SCState : SCStates)
+						{
+							if (!SCState->IsSourceControlled())
+							{
+								NewlyUncontrolledFiles.Add(SCState->GetFilename());
+							}
+						}
+
+						if (NewlyUncontrolledFiles.Num() > 0)
+						{
+							FScopedDisableSourceControl DisableSourceControl;
+							GEditor->AddDeferredMarkForAddFiles(NewlyUncontrolledFiles);
+							GEditor->RunDeferredMarkForAddFiles();
+						}
+					}
+				}
 			}
 		});
 
@@ -1341,7 +1364,7 @@ bool USourceControlHelpers::RevertAndReloadPackages(const TArray<FString>& InPac
 			RevertOperation->SetRevertAll(true);
 		}
 
-		return SourceControlProvider.Execute(RevertOperation, InPackagesToRevert, EConcurrency::Synchronous, OperationCompleteCallback) == ECommandResult::Succeeded;
+		return ISourceControlModule::Get().GetProvider().Execute(RevertOperation, InPackagesToRevert, EConcurrency::Synchronous, OperationCompleteCallback) == ECommandResult::Succeeded;
 	};
 
 	return ApplyOperationAndReloadPackages(InPackagesToRevert,RevertOperation, bReloadWorld);
diff --git a/Engine/Source/Editor/UnrealEd/Classes/Editor/EditorEngine.h b/Engine/Source/Editor/UnrealEd/Classes/Editor/EditorEngine.h
--- a/Engine/Source/Editor/UnrealEd/Classes/Editor/EditorEngine.h
+++ b/Engine/Source/Editor/UnrealEd/Classes/Editor/EditorEngine.h
@@ -2230,6 +2230,11 @@ public:
 	 */
 	UNREALED_API void AddPackagesToDefaultChangelist(TArray<FString>& InPackageNames);
 
+	/**
+	 * Add the given file(s) to the deferred add list processed by RunDeferredMarkForAddFiles.
+	 */
+	UNREALED_API void AddDeferredMarkForAddFiles(TArrayView<const FString> Files);
+
 	/**
 	 * Will run a batch mark for add source control operation for files recently saved if source control is enabled
 	 * @param	bool	Unused param to for the function to be used as a delegate when source control dialog is closed.
diff --git a/Engine/Source/Editor/UnrealEd/Private/EditorEngine.cpp b/Engine/Source/Editor/UnrealEd/Private/EditorEngine.cpp
--- a/Engine/Source/Editor/UnrealEd/Private/EditorEngine.cpp
+++ b/Engine/Source/Editor/UnrealEd/Private/EditorEngine.cpp
@@ -4157,6 +4157,11 @@ bool UEditorEngine::IsPackageOKToSave(UPackage* InPackage, const FString& InFile
 	return true;
 }
 
+void UEditorEngine::AddDeferredMarkForAddFiles(TArrayView<const FString> Files)
+{
+	DeferredFilesToAddToSourceControl.Append(Files);
+}
+
 void UEditorEngine::RunDeferredMarkForAddFiles(bool)
 {
 	if (DeferredFilesToAddToSourceControl.IsEmpty())
