CIS fix

#rnx

diff --git a/Engine/Source/Developer/SourceControl/Private/SourceControlHelpers.cpp b/Engine/Source/Developer/SourceControl/Private/SourceControlHelpers.cpp
--- a/Engine/Source/Developer/SourceControl/Private/SourceControlHelpers.cpp
+++ b/Engine/Source/Developer/SourceControl/Private/SourceControlHelpers.cpp
@@ -1338,7 +1338,7 @@ bool USourceControlHelpers::RevertAndReloadPackages(const TArray<FString>& InPac
 					if (ISourceControlModule::Get().GetProvider().GetState(InPackagesToRevert, SCStates, EStateCacheUsage::ForceUpdate) == ECommandResult::Succeeded)
 					{
 						TArray<FString> NewlyUncontrolledFiles;
-						for (const FSourceControlStateRef SCState : SCStates)
+						for (const FSourceControlStateRef& SCState : SCStates)
 						{
 							if (!SCState->IsSourceControlled())
 							{
