FSourceControlModule : Remove the close button from the Revision Control Login window.

In changelist 37621937 we made a fix to avoid leaving the provider in a bad state after the user entered invalid settings in the SSourceControlLogin and then clicked the close button.  However, that change was also clearing the provider in the case where the user opens and then closes the login window without changing any settings, leading to a confusing UX.  The user needs to choose either "Accept Settings" or "Disable Revision Control", so we are removing the close button to avoid ambiguity.

diff --git a/Engine/Source/Developer/SourceControl/Private/SSourceControlLogin.h b/Engine/Source/Developer/SourceControl/Private/SSourceControlLogin.h
--- a/Engine/Source/Developer/SourceControl/Private/SSourceControlLogin.h
+++ b/Engine/Source/Developer/SourceControl/Private/SSourceControlLogin.h
@@ -50,11 +50,6 @@ public:
 	 */
 	void RefreshSettings();
 
-	bool IsConnected() const
-	{
-		return ConnectionState == ELoginConnectionState::Connected;
-	}
-
 private:
 
 	/** Ticks the source control module (only necessary if the login window is modal) */
diff --git a/Engine/Source/Developer/SourceControl/Private/SourceControlModule.cpp b/Engine/Source/Developer/SourceControl/Private/SourceControlModule.cpp
--- a/Engine/Source/Developer/SourceControl/Private/SourceControlModule.cpp
+++ b/Engine/Source/Developer/SourceControl/Private/SourceControlModule.cpp
@@ -175,7 +175,7 @@ void FSourceControlModule::ShowLoginDialog(const FSourceControlLoginClosed& InOn
 		// Create the window
 		SourceControlLoginWindowPtr = SNew(SWindow)
 			.Title( LOCTEXT("SourceControlLoginTitle", "Revision Control Login") )
-			.HasCloseButton(true)
+			.HasCloseButton(false)
 			.SupportsMaximize(false) 
 			.SupportsMinimize(false)
 			.SizingRule( ESizingRule::Autosized );
@@ -220,14 +220,7 @@ void FSourceControlModule::ShowLoginDialog(const FSourceControlLoginClosed& InOn
 }
 
 void FSourceControlModule::OnSourceControlDialogClosed(const TSharedRef<SWindow>& InWindow)
-{
-#if SOURCE_CONTROL_WITH_SLATE
-	if (!SourceControlLoginPtr->IsConnected())
-	{
-		SetProvider("None");
-	}
-#endif
-	
+{	
 	SourceControlLoginWindowPtr = NULL;
 	SourceControlLoginPtr = NULL;
 	--TemporarilyDisabledCount;
