OFPE: Fix source control changelists window showing redirector files as entities after renaming entities by refreshing the source control asset data cache when renaming objects in the editor.

diff --git a/Engine/Plugins/Developer/PerforceSourceControl/PerforceSourceControl.uplugin b/Engine/Plugins/Developer/PerforceSourceControl/PerforceSourceControl.uplugin
--- a/Engine/Plugins/Developer/PerforceSourceControl/PerforceSourceControl.uplugin
+++ b/Engine/Plugins/Developer/PerforceSourceControl/PerforceSourceControl.uplugin
@@ -19,7 +19,7 @@
 		{
 			"Name" : "PerforceSourceControl",
 			"Type" : "UncookedOnly",
-			"LoadingPhase" : "EarliestPossible"
+			"LoadingPhase" : "PreEarlyLoadingScreen"
 		}
 	]
 }
\ No newline at end of file
diff --git a/Engine/Plugins/Developer/PlasticSourceControl/PlasticSourceControl.uplugin b/Engine/Plugins/Developer/PlasticSourceControl/PlasticSourceControl.uplugin
--- a/Engine/Plugins/Developer/PlasticSourceControl/PlasticSourceControl.uplugin
+++ b/Engine/Plugins/Developer/PlasticSourceControl/PlasticSourceControl.uplugin
@@ -18,7 +18,7 @@
 		{
 			"Name": "PlasticSourceControl",
 			"Type": "UncookedOnly",
-			"LoadingPhase": "EarliestPossible",
+			"LoadingPhase": "PreEarlyLoadingScreen",
 			"PlatformAllowList": [
 				"Win64",
 				"Linux"
diff --git a/Engine/Source/Developer/SourceControl/Private/SourceControlAssetDataCache.cpp b/Engine/Source/Developer/SourceControl/Private/SourceControlAssetDataCache.cpp
--- a/Engine/Source/Developer/SourceControl/Private/SourceControlAssetDataCache.cpp
+++ b/Engine/Source/Developer/SourceControl/Private/SourceControlAssetDataCache.cpp
@@ -4,6 +4,7 @@
 
 #include "Async/Async.h"
 #include "AssetRegistry/AssetData.h"
+#include "AssetRegistry/AssetRegistryModule.h"
 #include "HAL/IConsoleManager.h"
 #include "Misc/CoreDelegates.h"
 #include "AssetRegistry/IAssetRegistry.h"
@@ -32,21 +33,24 @@ void FSourceControlAssetDataCache::Startup()
 	}));
 
 #if WITH_EDITOR
-	ActorLabelChangedDelegateHandle = FCoreDelegates::OnActorLabelChanged.AddLambda([](AActor* ChangedActor)
+	auto ClearCachedAssetData = [](const UObject* Object)
 	{
-		if (ensure(ChangedActor))
+		// Transient packages should not be saved, therefore we can skip the asset data cache update for them to mitigate the perf impact of this call 
+		if (UPackage* Package = Object->GetPackage(); Package && !Package->HasAnyFlags(RF_Transient))
 		{
-			if (UPackage* Package = ChangedActor->GetPackage())
-			{
-				// Transient packages should not be saved, therefore we can skip the asset data cache update for them to mitigate the perf impact of this call 
-				if (Package->HasAnyFlags(RF_Transient))
-				{
-					return;
-				}
-
-				FString Filename = USourceControlHelpers::PackageFilename(Package);
-				ISourceControlModule::Get().GetAssetDataCache().ClearAssetData(Filename);
-			}
+			FString Filename = USourceControlHelpers::PackageFilename(Package);
+			ISourceControlModule::Get().GetAssetDataCache().ClearAssetData(Filename);
+		}
+	};
+
+	ActorLabelChangedDelegateHandle = FCoreDelegates::OnActorLabelChanged.AddLambda(ClearCachedAssetData);
+
+	FAssetRegistryModule& AssetRegistryModule = FModuleManager::LoadModuleChecked<FAssetRegistryModule>(FName("AssetRegistry"));
+	AssetUpdatedDelegateHandle = AssetRegistryModule.Get().OnAssetUpdated().AddLambda([&ClearCachedAssetData](const FAssetData& AssetData)
+	{
+		if (UObject* Object = AssetData.FastGetAsset())
+		{
+			ClearCachedAssetData(Object);
 		}
 	});
 #endif
@@ -58,6 +62,11 @@ void FSourceControlAssetDataCache::Shutdown()
 {
 #if WITH_EDITOR
 	FCoreDelegates::OnActorLabelChanged.Remove(ActorLabelChangedDelegateHandle);
+
+	if (FAssetRegistryModule* AssetRegistryModule = FModuleManager::GetModulePtr<FAssetRegistryModule>(TEXT("AssetRegistry")))
+	{
+		AssetRegistryModule->Get().OnAssetUpdated().Remove(AssetUpdatedDelegateHandle);
+	}
 #endif
 
 	ISourceControlModule& SCCModule = ISourceControlModule::Get();
diff --git a/Engine/Source/Developer/SourceControl/Public/SourceControlAssetDataCache.h b/Engine/Source/Developer/SourceControl/Public/SourceControlAssetDataCache.h
--- a/Engine/Source/Developer/SourceControl/Public/SourceControlAssetDataCache.h
+++ b/Engine/Source/Developer/SourceControl/Public/SourceControlAssetDataCache.h
@@ -113,6 +113,7 @@ private:
 	TQueue<FAssetDataToLoad, EQueueMode::Mpsc>			AssetDataToLoad;
 	FDelegateHandle										ProviderChangedDelegateHandle;
 	FDelegateHandle										ActorLabelChangedDelegateHandle;
+	FDelegateHandle										AssetUpdatedDelegateHandle;
 	std::atomic<uint32>									CurrentAsyncTask = 0;
 	bool												bIsSourceControlDialogShown;
 };
