OFPE: Fix source control changelists window showing redirector files as entities after renaming entities by refreshing the source control asset data cache when renaming objects in the editor.
Bring back CL39066763 with fix: instead of changing when the source control plugins initializes to make sure FSourceControlAssetDataCache can attach to OnAssetUpdated, FSourceControlAssetDataCache will now wait until OnAllModuleLoadingPhasesComplete is fired to register to OnAssetUpdated,

diff --git a/Engine/Source/Developer/SourceControl/Private/SourceControlAssetDataCache.cpp b/Engine/Source/Developer/SourceControl/Private/SourceControlAssetDataCache.cpp
--- a/Engine/Source/Developer/SourceControl/Private/SourceControlAssetDataCache.cpp
+++ b/Engine/Source/Developer/SourceControl/Private/SourceControlAssetDataCache.cpp
@@ -4,6 +4,7 @@
 
 #include "Async/Async.h"
 #include "AssetRegistry/AssetData.h"
+#include "AssetRegistry/AssetRegistryModule.h"
 #include "HAL/IConsoleManager.h"
 #include "Misc/CoreDelegates.h"
 #include "AssetRegistry/IAssetRegistry.h"
@@ -32,22 +33,28 @@ void FSourceControlAssetDataCache::Startup()
 	}));
 
 #if WITH_EDITOR
-	ActorLabelChangedDelegateHandle = FCoreDelegates::OnActorLabelChanged.AddLambda([](AActor* ChangedActor)
+	auto ClearCachedAssetData = [](const UObject* Object)
 	{
-		if (ensure(ChangedActor))
+		// Transient packages should not be saved, therefore we can skip the asset data cache update for them to mitigate the perf impact of this call 
+		if (UPackage* Package = Object->GetPackage(); Package && !Package->HasAnyFlags(RF_Transient))
 		{
-			if (UPackage* Package = ChangedActor->GetPackage())
+			FString Filename = USourceControlHelpers::PackageFilename(Package);
+			ISourceControlModule::Get().GetAssetDataCache().ClearAssetData(Filename);
+		}
+	};
+
+	ActorLabelChangedDelegateHandle = FCoreDelegates::OnActorLabelChanged.AddLambda(ClearCachedAssetData);
+
+	FCoreDelegates::OnAllModuleLoadingPhasesComplete.AddLambda([this, &ClearCachedAssetData]()
+	{
+		FAssetRegistryModule& AssetRegistryModule = FModuleManager::LoadModuleChecked<FAssetRegistryModule>(FName("AssetRegistry"));
+		AssetUpdatedDelegateHandle = AssetRegistryModule.Get().OnAssetUpdated().AddLambda([&ClearCachedAssetData](const FAssetData& AssetData)
+		{
+			if (UObject* Object = AssetData.FastGetAsset())
 			{
-				// Transient packages should not be saved, therefore we can skip the asset data cache update for them to mitigate the perf impact of this call 
-				if (Package->HasAnyFlags(RF_Transient))
-				{
-					return;
-				}
-
-				FString Filename = USourceControlHelpers::PackageFilename(Package);
-				ISourceControlModule::Get().GetAssetDataCache().ClearAssetData(Filename);
+				ClearCachedAssetData(Object);
 			}
-		}
+		});
 	});
 #endif
 
@@ -58,6 +65,11 @@ void FSourceControlAssetDataCache::Shutdown()
 {
 #if WITH_EDITOR
 	FCoreDelegates::OnActorLabelChanged.Remove(ActorLabelChangedDelegateHandle);
+
+	if (FAssetRegistryModule* AssetRegistryModule = FModuleManager::GetModulePtr<FAssetRegistryModule>(TEXT("AssetRegistry")))
+	{
+		AssetRegistryModule->Get().OnAssetUpdated().Remove(AssetUpdatedDelegateHandle);
+	}
 #endif
 
 	ISourceControlModule& SCCModule = ISourceControlModule::Get();
diff --git a/Engine/Source/Developer/SourceControl/Public/SourceControlAssetDataCache.h b/Engine/Source/Developer/SourceControl/Public/SourceControlAssetDataCache.h
--- a/Engine/Source/Developer/SourceControl/Public/SourceControlAssetDataCache.h
+++ b/Engine/Source/Developer/SourceControl/Public/SourceControlAssetDataCache.h
@@ -113,6 +113,7 @@ private:
 	TQueue<FAssetDataToLoad, EQueueMode::Mpsc>			AssetDataToLoad;
 	FDelegateHandle										ProviderChangedDelegateHandle;
 	FDelegateHandle										ActorLabelChangedDelegateHandle;
+	FDelegateHandle										AssetUpdatedDelegateHandle;
 	std::atomic<uint32>									CurrentAsyncTask = 0;
 	bool												bIsSourceControlDialogShown;
 };
