[Revision Control] Add batch versions of get status and provide information on branch status to the BPFL SourceControlHelpers

diff --git a/Engine/Source/Developer/SourceControl/Private/SourceControlHelpers.cpp b/Engine/Source/Developer/SourceControl/Private/SourceControlHelpers.cpp
--- a/Engine/Source/Developer/SourceControl/Private/SourceControlHelpers.cpp
+++ b/Engine/Source/Developer/SourceControl/Private/SourceControlHelpers.cpp
@@ -238,6 +238,55 @@ bool ConvertFilesToQualifiedPaths(const TArray<FString>& InFiles, TArray<FString
 	return true;
 }
 
+void RunQueryFileStates(TConstArrayView<FString> InFiles, EConcurrency::Type InConcurrencyType, bool bSilent, TFunction<void(FSourceControlState)> FileStateCallback)
+{
+	// Ensure source control system is up and running
+	ISourceControlProvider* Provider = SourceControlHelpersInternal::VerifySourceControl(bSilent);
+	if (!Provider)
+	{
+		return;
+	}
+
+	TArray<FString> SCFilesToUpdate;
+	for (const FString& File : InFiles)
+	{
+		// Determine file type and ensure it is in form source control wants
+		FString SCFile = SourceControlHelpersInternal::ConvertFileToQualifiedPath(File, bSilent);
+
+		if (SCFile.IsEmpty())
+		{
+			// Improper or invalid SCC state
+			FFormatNamedArguments Arguments;
+			Arguments.Add(TEXT("InFile"), FText::FromString(File));
+			SourceControlHelpersInternal::LogError(FText::Format(LOCTEXT("CouldNotDetermineState", "Could not resolve file path for '{InFile}'."), Arguments), bSilent);
+			continue;
+		}
+		SCFilesToUpdate.Emplace(MoveTemp(SCFile));
+	}
+
+	// Make sure we update the modified state of the files (Perforce requires this
+	// since can be a more expensive test).
+	TSharedRef<FUpdateStatus, ESPMode::ThreadSafe> UpdateStatusOperation = ISourceControlOperation::Create<FUpdateStatus>();
+	UpdateStatusOperation->SetUpdateModifiedState(true);
+
+	ISourceControlModule::Get().GetProvider().Execute(
+		UpdateStatusOperation,
+		SCFilesToUpdate,
+		InConcurrencyType,
+		FSourceControlOperationComplete::CreateLambda([FileStateCallback, Provider, SCFilesToUpdate, bSilent](const FSourceControlOperationRef& InOperation, ECommandResult::Type InResult)
+			{
+				TArray<FSourceControlStateRef> SCStates;
+				Provider->GetState(SCFilesToUpdate, SCStates, EStateCacheUsage::Use);
+
+				for (const FSourceControlStatePtr SCState : SCStates)
+				{
+					FSourceControlState State;
+					State.SetFromStatus(SCState);
+					FileStateCallback(State);
+				}
+			}));
+}
+
 }  // namespace SourceControlHelpersInternal
 
 
@@ -1578,146 +1627,44 @@ bool USourceControlHelpers::CopyFile(const FString& InSourcePath, const FString&
 	return Result == ECommandResult::Succeeded;
 }
 
-FSourceControlState USourceControlHelpers::QueryFileState(const FString& InFile, bool bSilent)
+TArray<FSourceControlState> USourceControlHelpers::QueryFileStates(const TArray<FString> InFiles, bool bSilent)
 {
-	FSourceControlState State;
-
-	State.bIsValid = false;
-
-	// Determine file type and ensure it is in form source control wants
-	FString SCFile = SourceControlHelpersInternal::ConvertFileToQualifiedPath(InFile, bSilent);
-
-	if (SCFile.IsEmpty())
-	{
-		State.Filename = InFile;
-		return State;
-	}
-
-	State.Filename = SCFile;
-
-	// Ensure source control system is up and running
-	ISourceControlProvider* Provider = SourceControlHelpersInternal::VerifySourceControl(bSilent);
-
-	if (!Provider)
-	{
-		return State;
-	}
-
-	// Make sure we update the modified state of the files (Perforce requires this
-	// since can be a more expensive test).
-	TSharedRef<FUpdateStatus, ESPMode::ThreadSafe> UpdateStatusOperation = ISourceControlOperation::Create<FUpdateStatus>();
-	UpdateStatusOperation->SetUpdateModifiedState(true);
-	Provider->Execute(UpdateStatusOperation, SCFile);
+	TArray<FSourceControlState> States;
 
-	FSourceControlStatePtr SCState = Provider->GetState(SCFile, EStateCacheUsage::Use);
-
-	if (!SCState.IsValid())
-	{
-		// Improper or invalid SCC state
-		FFormatNamedArguments Arguments;
-		Arguments.Add(TEXT("InFile"), FText::FromString(InFile));
-		Arguments.Add(TEXT("SCFile"), FText::FromString(SCFile));
-		SourceControlHelpersInternal::LogError(FText::Format(LOCTEXT("CouldNotDetermineState", "Could not determine revision control state of file '{InFile}' ({SCFile})."), Arguments), bSilent);
-
-		return State;
-	}
-
-	// Return FSourceControlState rather than a ISourceControlState directly so that
-	// scripting systems can access it.
-
-	State.bIsValid = true;
+	SourceControlHelpersInternal::RunQueryFileStates(InFiles, EConcurrency::Synchronous, bSilent, 
+		[&States](FSourceControlState FileStateOut)
+		{
+			States.Emplace(MoveTemp(FileStateOut));
+		});
 
-	// Copy over state info
-	// - make these assignments a method of FSourceControlState if anything else sets a state
-	State.bIsUnknown			= SCState->IsUnknown();
-	State.bIsSourceControlled	= SCState->IsSourceControlled();
-	State.bCanCheckIn			= SCState->CanCheckIn();
-	State.bCanCheckOut			= SCState->CanCheckout();
-	State.bIsCheckedOut			= SCState->IsCheckedOut();
-	State.bIsCurrent			= SCState->IsCurrent();
-	State.bIsAdded				= SCState->IsAdded();
-	State.bIsDeleted			= SCState->IsDeleted();
-	State.bIsIgnored			= SCState->IsIgnored();
-	State.bCanEdit				= SCState->CanEdit();
-	State.bCanDelete			= SCState->CanDelete();
-	State.bCanAdd				= SCState->CanAdd();
-	State.bIsConflicted			= SCState->IsConflicted();
-	State.bCanRevert			= SCState->CanRevert();
-	State.bIsModified			= SCState->IsModified();
-	State.bIsCheckedOutOther	= SCState->IsCheckedOutOther();
+	return States;
+}
 
-	if (State.bIsCheckedOutOther)
-	{
-		SCState->IsCheckedOutOther(&State.CheckedOutOther);
-	}
+FSourceControlState USourceControlHelpers::QueryFileState(const FString& InFile, bool bSilent)
+{
+	FSourceControlState State;
+	State.Filename = InFile;
+	SourceControlHelpersInternal::RunQueryFileStates(TArray<FString> {InFile}, EConcurrency::Synchronous, bSilent,
+		[&State](FSourceControlState FileStateOut)
+		{
+			State = MoveTemp(FileStateOut);
+		});
 
 	return State;
 }
 
-void USourceControlHelpers::AsyncQueryFileState(FQueryFileStateDelegate FileStateCallback, const FString& InFile, bool bSilent)
+void USourceControlHelpers::AsyncQueryFileStates(FQueryFileStateDelegate FileStateCallback, TArray<FString> InFiles, bool bSilent)
 {
-	// Determine file type and ensure it is in form source control wants
-	FString SCFile = SourceControlHelpersInternal::ConvertFileToQualifiedPath(InFile, bSilent);
-
-	if (SCFile.IsEmpty())
-	{
-		return;
-	}
-
-	if (ISourceControlProvider* Provider = SourceControlHelpersInternal::VerifySourceControl(bSilent))
-	{
-		// Make sure we update the modified state of the files (Perforce requires this
-		// since can be a more expensive test).
-		TSharedRef<FUpdateStatus, ESPMode::ThreadSafe> UpdateStatusOperation = ISourceControlOperation::Create<FUpdateStatus>();
-		UpdateStatusOperation->SetUpdateModifiedState(true);
-
-		ISourceControlModule::Get().GetProvider().Execute(UpdateStatusOperation, SCFile, EConcurrency::Asynchronous,
-			FSourceControlOperationComplete::CreateLambda([FileStateCallback, Provider, SCFile, InFile, bSilent](const FSourceControlOperationRef& InOperation, ECommandResult::Type InResult)
+	SourceControlHelpersInternal::RunQueryFileStates(InFiles, EConcurrency::Asynchronous, bSilent,
+		[FileStateCallback](FSourceControlState InState)
 		{
-			FSourceControlState State;
-			const FSourceControlStatePtr SCState = Provider->GetState(SCFile, EStateCacheUsage::Use);
-
-			if (!SCState.IsValid())
-			{
-				// Improper or invalid SCC state
-				FFormatNamedArguments Arguments;
-				Arguments.Add(TEXT("InFile"), FText::FromString(InFile));
-				Arguments.Add(TEXT("SCFile"), FText::FromString(SCFile));
-				SourceControlHelpersInternal::LogError(FText::Format(LOCTEXT("CouldNotDetermineState", "Could not determine revision control state of file '{InFile}' ({SCFile})."), Arguments), bSilent);
-
-				FileStateCallback.ExecuteIfBound(State);
-				return;
-			}
-
-			State.bIsValid = true;
-
-			// Copy over state info
-			// - make these assignments a method of FSourceControlState if anything else sets a state
-			State.bIsUnknown = SCState->IsUnknown();
-			State.bIsSourceControlled = SCState->IsSourceControlled();
-			State.bCanCheckIn = SCState->CanCheckIn();
-			State.bCanCheckOut = SCState->CanCheckout();
-			State.bIsCheckedOut = SCState->IsCheckedOut();
-			State.bIsCurrent = SCState->IsCurrent();
-			State.bIsAdded = SCState->IsAdded();
-			State.bIsDeleted = SCState->IsDeleted();
-			State.bIsIgnored = SCState->IsIgnored();
-			State.bCanEdit = SCState->CanEdit();
-			State.bCanDelete = SCState->CanDelete();
-			State.bCanAdd = SCState->CanAdd();
-			State.bIsConflicted = SCState->IsConflicted();
-			State.bCanRevert = SCState->CanRevert();
-			State.bIsModified = SCState->IsModified();
-			State.bIsCheckedOutOther = SCState->IsCheckedOutOther();
-
-			if (State.bIsCheckedOutOther)
-			{
-				SCState->IsCheckedOutOther(&State.CheckedOutOther);
-			}
+			FileStateCallback.ExecuteIfBound(InState);
+		});
+}
 
-			FileStateCallback.ExecuteIfBound(State);
-		}));
-	}
+void USourceControlHelpers::AsyncQueryFileState(FQueryFileStateDelegate FileStateCallback, const FString& InFile, bool bSilent)
+{
+	AsyncQueryFileStates(FileStateCallback, TArray<FString> {InFile}, bSilent);
 }
 
 bool USourceControlHelpers::GetFilesInDepotAtPath(const FString& Path, TArray<FString>& OutFilesList, bool bIncludeDeleted, bool bSilent, bool bIsFileRegexSearch)
diff --git a/Engine/Source/Developer/SourceControl/Public/SourceControlHelpers.h b/Engine/Source/Developer/SourceControl/Public/SourceControlHelpers.h
--- a/Engine/Source/Developer/SourceControl/Public/SourceControlHelpers.h
+++ b/Engine/Source/Developer/SourceControl/Public/SourceControlHelpers.h
@@ -58,13 +58,52 @@ public:
 		bIsCheckedOutOther(false)
 	{}
 
+	void SetFromStatus(const FSourceControlStatePtr& InStatus)
+	{
+		if (InStatus)
+		{
+			bIsValid = true;
+
+			// Copy over state info
+			// - make these assignments a method of FSourceControlState if anything else sets a state
+			bIsUnknown = InStatus->IsUnknown();
+			bIsSourceControlled = InStatus->IsSourceControlled();
+			bCanCheckIn = InStatus->CanCheckIn();
+			bCanCheckOut = InStatus->CanCheckout();
+			bIsCheckedOut = InStatus->IsCheckedOut();
+			bIsCurrent = InStatus->IsCurrent();
+			bIsAdded = InStatus->IsAdded();
+			bIsDeleted = InStatus->IsDeleted();
+			bIsIgnored = InStatus->IsIgnored();
+			bCanEdit = InStatus->CanEdit();
+			bCanDelete = InStatus->CanDelete();
+			bCanAdd = InStatus->CanAdd();
+			bIsConflicted = InStatus->IsConflicted();
+			bCanRevert = InStatus->CanRevert();
+			bIsModified = InStatus->IsModified();
+			bIsCheckedOutOther = InStatus->IsCheckedOutOther();
+			bIsCheckedOutInOtherBranch = InStatus->IsCheckedOutInOtherBranch();
+			bIsModifiedInOtherBranch = InStatus->IsModifiedInOtherBranch();
+			Filename = InStatus->GetFilename();
+
+			if (bIsCheckedOutOther)
+			{
+				InStatus->IsCheckedOutOther(&CheckedOutOther);
+			}
+		}
+		else
+		{
+			bIsValid = false;
+		}
+	}
+
 	/** Get the local filename that this state represents */
 	UPROPERTY(VisibleAnywhere, BlueprintReadOnly, Category = "Editor Scripting | Editor Revision Control Helpers", meta = (Keywords = "Source Control"))
 	FString Filename;
 
 	/** Indicates whether this source control state has valid information (true) or not (false) */
 	UPROPERTY(VisibleAnywhere, BlueprintReadOnly, Category = "Editor Scripting | Editor Revision Control Helpers", meta = (Keywords = "Source Control"))
-	bool bIsValid=false;
+	bool bIsValid = false;
 
 	/** Determine if we know anything about the source control state of this file */
 	UPROPERTY(VisibleAnywhere, BlueprintReadOnly, Category = "Editor Scripting | Editor Revision Control Helpers", meta = (Keywords = "Source Control"))
@@ -142,6 +181,13 @@ public:
 	UPROPERTY(VisibleAnywhere, BlueprintReadOnly, Category = "Editor Scripting | Editor Revision Control Helpers", meta = (Keywords = "Source Control"))
 	FString CheckedOutOther;
 
+	/** Determine if this file is checked out by someone else in another branch */
+	UPROPERTY(VisibleAnywhere, BlueprintReadOnly, Category = "Editor Scripting | Editor Revision Control Helpers", meta = (Keywords = "Source Control"))
+	bool bIsCheckedOutInOtherBranch = false;
+
+	/** Determine if this file is modified by someone in another branch */
+	UPROPERTY(VisibleAnywhere, BlueprintReadOnly, Category = "Editor Scripting | Editor Revision Control Helpers", meta = (Keywords = "Source Control"))
+	bool bIsModifiedInOtherBranch = false;
 };  // FSourceControlState
 
 
@@ -483,6 +529,9 @@ public:
 	UFUNCTION(BlueprintCallable, Category = "Editor Scripting | Editor Revision Control Helpers", meta = (Keywords = "Source Control"))
 	static FSourceControlState QueryFileState(const FString& InFile, bool bSilent = false);
 
+	UFUNCTION(BlueprintCallable, Category = "Editor Scripting | Editor Revision Control Helpers", meta = (Keywords = "Source Control"))
+	static TArray<FSourceControlState> QueryFileStates(TArray<FString> InFiles, bool bSilent = false);
+
 	//Delegate to broadcast FileState upon AsyncQueryFileState completion
 	DECLARE_DYNAMIC_DELEGATE_OneParam(FQueryFileStateDelegate, FSourceControlState, FileStateOut);
 	/**
@@ -494,6 +543,8 @@ public:
 	*/
 	UFUNCTION(BlueprintCallable, Category = "Editor Scripting | Editor Revision Control Helpers", meta = (Keywords = "Source Control"))
 	static void AsyncQueryFileState(FQueryFileStateDelegate FileStateCallback, const FString& InFile, bool bSilent = false);
+	UFUNCTION(BlueprintCallable, Category = "Editor Scripting | Editor Revision Control Helpers", meta = (Keywords = "Source Control"))
+	static void AsyncQueryFileStates(FQueryFileStateDelegate FileStateCallback, TArray<FString> InFiles, bool bSilent = false);
 	
 	/**
 	 * Use currently set source control provider to query the list of files in the depot under a certain path.
