[Revision Control] Fix loctext conflict

diff --git a/Engine/Source/Developer/SourceControl/Private/SourceControlHelpers.cpp b/Engine/Source/Developer/SourceControl/Private/SourceControlHelpers.cpp
--- a/Engine/Source/Developer/SourceControl/Private/SourceControlHelpers.cpp
+++ b/Engine/Source/Developer/SourceControl/Private/SourceControlHelpers.cpp
@@ -258,7 +258,7 @@ void RunQueryFileStates(TConstArrayView<FString> InFiles, EConcurrency::Type InC
 			// Improper or invalid SCC state
 			FFormatNamedArguments Arguments;
 			Arguments.Add(TEXT("InFile"), FText::FromString(File));
-			SourceControlHelpersInternal::LogError(FText::Format(LOCTEXT("CouldNotDetermineState", "Could not resolve file path for '{InFile}'."), Arguments), bSilent);
+			SourceControlHelpersInternal::LogError(FText::Format(LOCTEXT("CouldNotResolveFilepath", "Could not resolve file path for '{InFile}'."), Arguments), bSilent);
 			continue;
 		}
 		SCFilesToUpdate.Emplace(MoveTemp(SCFile));
