SourceControl: Avoid flushing all loads when requesting revert / unlink.

#rnx

diff --git a/Engine/Source/Developer/SourceControl/Private/SourceControlHelpers.cpp b/Engine/Source/Developer/SourceControl/Private/SourceControlHelpers.cpp
--- a/Engine/Source/Developer/SourceControl/Private/SourceControlHelpers.cpp
+++ b/Engine/Source/Developer/SourceControl/Private/SourceControlHelpers.cpp
@@ -967,6 +967,8 @@ bool USourceControlHelpers::ApplyOperationAndReloadPackages(const TArray<FString
 	TArray<FString> PackageFilenames;
 	TArray<FString> FilteredPackages;
 	TArray<FString> NotFoundPackages;
+	TArray<UPackage*> UnlinkPackages;
+
 	/** 
 	 * This is necessary to cover an edge case where the user reverts a deleted Level file.
 	 * In that case PackageFilename_Internal() will not be able to resolve to the correct extension
@@ -976,16 +978,6 @@ bool USourceControlHelpers::ApplyOperationAndReloadPackages(const TArray<FString
 
 	bool bSuccess = false;
 
-	auto DetachLinker = [](UPackage* Package)
-	{
-		if (!Package->IsFullyLoaded())
-		{
-			FlushAsyncLoading();
-			Package->FullyLoad();
-		}
-		ResetLoaders(Package);
-	};
-
 	// Normalize packagenames and filenames
 	for (const FString& Filename : InFilenames)
 	{
@@ -1010,7 +1002,7 @@ bool USourceControlHelpers::ApplyOperationAndReloadPackages(const TArray<FString
 	// If bReloadWorld=false, remove packages if they are loaded external packages or world
 	// If bReloadWorld=true, include world packages to reload
 	TSet<UPackage*> UniqueLoadedPackages;
-	PackageNames.RemoveAll([&FilteredPackages, &UniqueLoadedPackages, &NotFoundPackages, bReloadWorld, &DetachLinker](const FString& PackageName) -> bool
+	PackageNames.RemoveAll([&FilteredPackages, &UniqueLoadedPackages, &NotFoundPackages, &UnlinkPackages, bReloadWorld](const FString& PackageName) -> bool
 	{
 		UPackage* Package = FindPackage(NULL, *PackageName);
 		
@@ -1033,7 +1025,7 @@ bool USourceControlHelpers::ApplyOperationAndReloadPackages(const TArray<FString
 					if (bReloadWorld || !bIsAssetPartOfWorld)
 					{
 						// detach linker on the object
-						DetachLinker(Package);
+						UnlinkPackages.Emplace(Package);
 
 						// track its world for reloading - not the object package itself
 						if (bIsAssetPartOfWorld)
@@ -1133,10 +1125,27 @@ bool USourceControlHelpers::ApplyOperationAndReloadPackages(const TArray<FString
 
 	// Prepare the packages to be reverted...
 	TArray<UPackage*> LoadedPackages = UniqueLoadedPackages.Array();
-	for (UPackage* Package : LoadedPackages)
+
+	// Make sure they are unlinked...
+	UnlinkPackages.Append(LoadedPackages);
+	if (UnlinkPackages.Num() > 0)
 	{
-		// Detach the linkers of any loaded packages so that SCC can overwrite the files...
-		DetachLinker(Package);
+		TArray<UObject*> LoadedObjects;
+		TArray<int32> PendingPackageRequestIds;
+		LoadedObjects.Reserve(UnlinkPackages.Num());
+		PendingPackageRequestIds.Reserve(UnlinkPackages.Num());
+		for (UPackage* Package : UnlinkPackages)
+		{
+			LoadedObjects.Emplace(Package);
+
+			if (!Package->IsFullyLoaded())
+			{
+				PendingPackageRequestIds.Emplace(LoadPackageAsync(Package->GetName()));
+			}
+		}
+
+		FlushAsyncLoading(PendingPackageRequestIds);
+		ResetLoaders(LoadedObjects);
 	}
 
 	for (int32 PackageIndex = 0; PackageIndex < PackageNames.Num(); PackageIndex++)
