SourceControl: The 'SetForceUpdate' doesn't apply here.

It's being used by Perforce to apply the '-f' to the 'diff' command. However the 'diff' command is only executed by FUpdateStatus if 'SetShouldUpdateModifiedState' is set to 'true' which isn't the case here.

#rnx

diff --git a/Engine/Source/Developer/SourceControl/Private/SourceControlFileStatusMonitor.cpp b/Engine/Source/Developer/SourceControl/Private/SourceControlFileStatusMonitor.cpp
--- a/Engine/Source/Developer/SourceControl/Private/SourceControlFileStatusMonitor.cpp
+++ b/Engine/Source/Developer/SourceControl/Private/SourceControlFileStatusMonitor.cpp
@@ -324,7 +324,6 @@ bool FSourceControlFileStatusMonitor::Tick(float DeltaTime)
 		LastSourceControlCheckSecs = NowSecs;
 
 		TSharedRef<FUpdateStatus> UpdateStatusRequest = ISourceControlOperation::Create<FUpdateStatus>();
-		UpdateStatusRequest->SetForceUpdate(true);
 
 		ISourceControlModule::Get().GetProvider().Execute(UpdateStatusRequest, RequestedStatusFiles, EConcurrency::Asynchronous,
 			FSourceControlOperationComplete::CreateSP(this, &FSourceControlFileStatusMonitor::OnSourceControlStatusUpdate));
