PR #13343: Allow RevisionControlStyle to be overridden

diff --git a/Engine/Source/Developer/SourceControl/Private/RevisionControlStyle/RevisionControlStyle.cpp b/Engine/Source/Developer/SourceControl/Private/RevisionControlStyle/RevisionControlStyle.cpp
--- a/Engine/Source/Developer/SourceControl/Private/RevisionControlStyle/RevisionControlStyle.cpp
+++ b/Engine/Source/Developer/SourceControl/Private/RevisionControlStyle/RevisionControlStyle.cpp
@@ -17,6 +17,14 @@ FName FRevisionControlStyleManager::CurrentRevisionControlStyleName;
 
 FName FDefaultRevisionControlStyle::StyleName("DefaultRevisionControlStyle");
 
+FRevisionControlStyleManager::~FRevisionControlStyleManager()
+{
+	if (DefaultRevisionControlStyleInstance.IsValid())
+	{
+		FSlateStyleRegistry::UnRegisterSlateStyle(*DefaultRevisionControlStyleInstance);
+	}
+}
+
 // FRevisionControlStyleManager
 
 void FRevisionControlStyleManager::ReloadTextures()
@@ -26,19 +34,18 @@ void FRevisionControlStyleManager::ReloadTextures()
 
 const ISlateStyle& FRevisionControlStyleManager::Get()
 {
-	// If the default style instance doesn't exist (first call), create it and set it as active
-	if (!DefaultRevisionControlStyleInstance.IsValid())
-	{
-		DefaultRevisionControlStyleInstance = MakeShared<FDefaultRevisionControlStyle>();
-		SetActiveRevisionControlStyle(DefaultRevisionControlStyleInstance->GetStyleSetName());
-	}
-
 	// Find and get the currently active revision control style
 	if (const ISlateStyle* RevisionControlStyle = FSlateStyleRegistry::FindSlateStyle(CurrentRevisionControlStyleName))
 	{
 		return *RevisionControlStyle;
 	}
 
+	if (!DefaultRevisionControlStyleInstance.IsValid())
+	{
+		DefaultRevisionControlStyleInstance = MakeShared<FDefaultRevisionControlStyle>();
+		FSlateStyleRegistry::RegisterSlateStyle(*DefaultRevisionControlStyleInstance);
+	}
+
 	return *DefaultRevisionControlStyleInstance;
 }
 
@@ -60,16 +67,14 @@ void FRevisionControlStyleManager::SetActiveRevisionControlStyle(FName InNewActi
 
 void FRevisionControlStyleManager::ResetToDefaultRevisionControlStyle()
 {
-	SetActiveRevisionControlStyle(DefaultRevisionControlStyleInstance->GetStyleSetName());
+	SetActiveRevisionControlStyle(DefaultRevisionControlStyleInstance ? DefaultRevisionControlStyleInstance->GetStyleSetName() : NAME_None);
 }
 
 // FDefaultRevisionControlStyle
 
 FDefaultRevisionControlStyle::FDefaultRevisionControlStyle() : FSlateStyleSet(StyleName)
 {
-	FSlateStyleRegistry::RegisterSlateStyle(*this);
-
-	SetContentRoot(FPaths::EngineContentDir() / TEXT("Slate"));
+	SetCoreContentRoot(FPaths::EngineContentDir() / TEXT("Slate"));
 
 	// Any new revision control icons should be added here instead of in StarshipStyle. Any icons there that don't exist here should be moved here and used by calling FRevisionControlStyleManager::Get()
 	
@@ -85,103 +90,102 @@ FDefaultRevisionControlStyle::FDefaultRevisionControlStyle() : FSlateStyleSet(St
 	SnapshotHistoryRemoved = FLinearColor::FromSRGBColor(FColor::FromHex("#CD3642"));
 
 	// Status icons
-	Set("RevisionControl.Icon", new IMAGE_BRUSH_SVG("Starship/SourceControl/Status/RevisionControl", CoreStyleConstants::Icon16x16));
-	Set("RevisionControl.Icon.ConnectedBadge", new IMAGE_BRUSH_SVG("Starship/SourceControl/Status/RevisionControlBadgeConnected", CoreStyleConstants::Icon16x16, FStyleColors::Success));
-	Set("RevisionControl.Icon.WarningBadge", new IMAGE_BRUSH_SVG("Starship/SourceControl/Status/RevisionControlBadgeWarning", CoreStyleConstants::Icon16x16, FStyleColors::Warning));
+	Set("RevisionControl.Icon", new CORE_IMAGE_BRUSH_SVG("Starship/SourceControl/Status/RevisionControl", CoreStyleConstants::Icon16x16));
+	Set("RevisionControl.Icon.ConnectedBadge", new CORE_IMAGE_BRUSH_SVG("Starship/SourceControl/Status/RevisionControlBadgeConnected", CoreStyleConstants::Icon16x16, FStyleColors::Success));
+	Set("RevisionControl.Icon.WarningBadge", new CORE_IMAGE_BRUSH_SVG("Starship/SourceControl/Status/RevisionControlBadgeWarning", CoreStyleConstants::Icon16x16, FStyleColors::Warning));
 	
 	// Icons for revision control actions
 
-	Set("RevisionControl.Actions.Sync", new IMAGE_BRUSH_SVG("Starship/SourceControl/RC_Sync", CoreStyleConstants::Icon16x16));
-	Set("RevisionControl.Actions.CheckOut", new IMAGE_BRUSH_SVG("Starship/Common/check-circle", CoreStyleConstants::Icon16x16));
-	Set("RevisionControl.Actions.SyncAndCheckOut", new IMAGE_BRUSH_SVG("Starship/SourceControl/RC_SyncAndCheckOut", CoreStyleConstants::Icon16x16));
-	Set("RevisionControl.Actions.MakeWritable", new IMAGE_BRUSH_SVG("Starship/Common/edit", CoreStyleConstants::Icon16x16));
-	Set("RevisionControl.Actions.Add", new IMAGE_BRUSH_SVG("Starship/SourceControl/RC_MarkedForAdd", CoreStyleConstants::Icon16x16));
-	Set("RevisionControl.Actions.Submit", new IMAGE_BRUSH_SVG("Starship/SourceControl/RC_CheckIn", CoreStyleConstants::Icon16x16));
-	Set("RevisionControl.Actions.Connect", new IMAGE_BRUSH_SVG("Starship/SourceControl/Status/RevisionControl", CoreStyleConstants::Icon16x16));
-	Set("RevisionControl.Actions.History", new IMAGE_BRUSH_SVG("Starship/Common/Recent", CoreStyleConstants::Icon16x16));
-	Set("RevisionControl.Actions.Diff", new IMAGE_BRUSH_SVG("Starship/SourceControl/RC_Diff", CoreStyleConstants::Icon16x16));
-	Set("RevisionControl.Actions.Revert", new IMAGE_BRUSH_SVG("Starship/SourceControl/icon_SCC_Revert", CoreStyleConstants::Icon16x16));
-	Set("RevisionControl.Actions.Merge", new IMAGE_BRUSH_SVG("Starship/Common/Merge", CoreStyleConstants::Icon16x16));
-	Set("RevisionControl.Actions.Refresh", new IMAGE_BRUSH_SVG("Starship/Common/Update", CoreStyleConstants::Icon16x16));
-	Set("RevisionControl.Actions.ChangeSettings", new IMAGE_BRUSH_SVG("Starship/SourceControl/icon_SCC_Change_Source_Control_Settings", CoreStyleConstants::Icon16x16));
-	Set("RevisionControl.Actions.Rewind", new IMAGE_BRUSH_SVG("Starship/SourceControl/RC_Rewind", CoreStyleConstants::Icon16x16));
+	Set("RevisionControl.Actions.Sync", new CORE_IMAGE_BRUSH_SVG("Starship/SourceControl/RC_Sync", CoreStyleConstants::Icon16x16));
+	Set("RevisionControl.Actions.CheckOut", new CORE_IMAGE_BRUSH_SVG("Starship/Common/check-circle", CoreStyleConstants::Icon16x16));
+	Set("RevisionControl.Actions.SyncAndCheckOut", new CORE_IMAGE_BRUSH_SVG("Starship/SourceControl/RC_SyncAndCheckOut", CoreStyleConstants::Icon16x16));
+	Set("RevisionControl.Actions.MakeWritable", new CORE_IMAGE_BRUSH_SVG("Starship/Common/edit", CoreStyleConstants::Icon16x16));
+	Set("RevisionControl.Actions.Add", new CORE_IMAGE_BRUSH_SVG("Starship/SourceControl/RC_MarkedForAdd", CoreStyleConstants::Icon16x16));
+	Set("RevisionControl.Actions.Submit", new CORE_IMAGE_BRUSH_SVG("Starship/SourceControl/RC_CheckIn", CoreStyleConstants::Icon16x16));
+	Set("RevisionControl.Actions.Connect", new CORE_IMAGE_BRUSH_SVG("Starship/SourceControl/Status/RevisionControl", CoreStyleConstants::Icon16x16));
+	Set("RevisionControl.Actions.History", new CORE_IMAGE_BRUSH_SVG("Starship/Common/Recent", CoreStyleConstants::Icon16x16));
+	Set("RevisionControl.Actions.Diff", new CORE_IMAGE_BRUSH_SVG("Starship/SourceControl/RC_Diff", CoreStyleConstants::Icon16x16));
+	Set("RevisionControl.Actions.Revert", new CORE_IMAGE_BRUSH_SVG("Starship/SourceControl/icon_SCC_Revert", CoreStyleConstants::Icon16x16));
+	Set("RevisionControl.Actions.Merge", new CORE_IMAGE_BRUSH_SVG("Starship/Common/Merge", CoreStyleConstants::Icon16x16));
+	Set("RevisionControl.Actions.Refresh", new CORE_IMAGE_BRUSH_SVG("Starship/Common/Update", CoreStyleConstants::Icon16x16));
+	Set("RevisionControl.Actions.ChangeSettings", new CORE_IMAGE_BRUSH_SVG("Starship/SourceControl/icon_SCC_Change_Source_Control_Settings", CoreStyleConstants::Icon16x16));
+	Set("RevisionControl.Actions.Rewind", new CORE_IMAGE_BRUSH_SVG("Starship/SourceControl/RC_Rewind", CoreStyleConstants::Icon16x16));
 
 	// Icons representing the various revision control states
 
-	Set("RevisionControl.Locked", new IMAGE_BRUSH_SVG("Starship/Common/lock", CoreStyleConstants::Icon16x16));
-	Set("RevisionControl.Unlocked", new IMAGE_BRUSH_SVG("Starship/Common/lock-unlocked", CoreStyleConstants::Icon16x16));
+	Set("RevisionControl.Locked", new CORE_IMAGE_BRUSH_SVG("Starship/Common/lock", CoreStyleConstants::Icon16x16));
+	Set("RevisionControl.Unlocked", new CORE_IMAGE_BRUSH_SVG("Starship/Common/lock-unlocked", CoreStyleConstants::Icon16x16));
 
-	Set("RevisionControl.Shelved", new IMAGE_BRUSH_SVG("Starship/SourceControl/SCC_Shelved", CoreStyleConstants::Icon16x16));
-	Set("RevisionControl.CheckedOut", new IMAGE_BRUSH_SVG("Starship/SourceControl/SCC_CheckedOut", CoreStyleConstants::Icon16x16, StatusCheckedOutColor));
+	Set("RevisionControl.Shelved", new CORE_IMAGE_BRUSH_SVG("Starship/SourceControl/SCC_Shelved", CoreStyleConstants::Icon16x16));
+	Set("RevisionControl.CheckedOut", new CORE_IMAGE_BRUSH_SVG("Starship/SourceControl/SCC_CheckedOut", CoreStyleConstants::Icon16x16, StatusCheckedOutColor));
 	
-	Set("RevisionControl.OpenForAdd", new IMAGE_BRUSH_SVG("Starship/SourceControl/RC_MarkedForAdd", CoreStyleConstants::Icon16x16, FStyleColors::AccentBlue));
-	Set("RevisionControl.MarkedForDelete", new IMAGE_BRUSH_SVG("Starship/SourceControl/SCC_MarkedForDelete", CoreStyleConstants::Icon16x16, FStyleColors::Error));
+	Set("RevisionControl.OpenForAdd", new CORE_IMAGE_BRUSH_SVG("Starship/SourceControl/RC_MarkedForAdd", CoreStyleConstants::Icon16x16, FStyleColors::AccentBlue));
+	Set("RevisionControl.MarkedForDelete", new CORE_IMAGE_BRUSH_SVG("Starship/SourceControl/SCC_MarkedForDelete", CoreStyleConstants::Icon16x16, FStyleColors::Error));
 
-	Set("RevisionControl.CheckedOutByOtherUser", new IMAGE_BRUSH_SVG("Starship/SourceControl/RC_CheckedOther", CoreStyleConstants::Icon16x16, StatusCheckedOutByOtherUserColor));
-	Set("RevisionControl.CheckedOutByOtherUserBadge", new IMAGE_BRUSH_SVG("Starship/SourceControl/RC_CheckedOtherBadge", CoreStyleConstants::Icon16x16, FStyleColors::AccentBlue));
+	Set("RevisionControl.CheckedOutByOtherUser", new CORE_IMAGE_BRUSH_SVG("Starship/SourceControl/RC_CheckedOther", CoreStyleConstants::Icon16x16, StatusCheckedOutByOtherUserColor));
+	Set("RevisionControl.CheckedOutByOtherUserBadge", new CORE_IMAGE_BRUSH_SVG("Starship/SourceControl/RC_CheckedOtherBadge", CoreStyleConstants::Icon16x16, FStyleColors::AccentBlue));
 
-	Set("RevisionControl.CheckedOutByOtherUserOtherBranch", new IMAGE_BRUSH_SVG("Starship/SourceControl/RC_CheckedBranch", CoreStyleConstants::Icon16x16, BranchedColor));
-	Set("RevisionControl.CheckedOutByOtherUserOtherBranchBadge", new IMAGE_BRUSH_SVG("Starship/SourceControl/RC_CheckedOtherBadge", CoreStyleConstants::Icon16x16, FStyleColors::Warning));
+	Set("RevisionControl.CheckedOutByOtherUserOtherBranch", new CORE_IMAGE_BRUSH_SVG("Starship/SourceControl/RC_CheckedBranch", CoreStyleConstants::Icon16x16, BranchedColor));
+	Set("RevisionControl.CheckedOutByOtherUserOtherBranchBadge", new CORE_IMAGE_BRUSH_SVG("Starship/SourceControl/RC_CheckedOtherBadge", CoreStyleConstants::Icon16x16, FStyleColors::Warning));
 
-	Set("RevisionControl.ModifiedOtherBranch", new IMAGE_BRUSH_SVG("Starship/SourceControl/SCC_ModifiedOtherBranch", CoreStyleConstants::Icon16x16, BranchedColor));
-	Set("RevisionControl.ModifiedBadge", new IMAGE_BRUSH_SVG("Starship/SourceControl/RC_BranchModifiedBadge", CoreStyleConstants::Icon16x16, FStyleColors::Warning));
-	Set("RevisionControl.ModifiedLocally", new IMAGE_BRUSH_SVG("Starship/SourceControl/RC_ModifiedLocally", CoreStyleConstants::Icon16x16, FStyleColors::Foreground));
+	Set("RevisionControl.ModifiedOtherBranch", new CORE_IMAGE_BRUSH_SVG("Starship/SourceControl/SCC_ModifiedOtherBranch", CoreStyleConstants::Icon16x16, BranchedColor));
+	Set("RevisionControl.ModifiedBadge", new CORE_IMAGE_BRUSH_SVG("Starship/SourceControl/RC_BranchModifiedBadge", CoreStyleConstants::Icon16x16, FStyleColors::Warning));
+	Set("RevisionControl.ModifiedLocally", new CORE_IMAGE_BRUSH_SVG("Starship/SourceControl/RC_ModifiedLocally", CoreStyleConstants::Icon16x16, FStyleColors::Foreground));
 
-	Set("RevisionControl.NotAtHeadRevision", new IMAGE_BRUSH_SVG("Starship/SourceControl/RC_NewerVersion", CoreStyleConstants::Icon16x16, StatusNotAtHeadRevisionColor));
-	Set("RevisionControl.NotInDepot", new IMAGE_BRUSH_SVG("Starship/SourceControl/SCC_NotInDepot", CoreStyleConstants::Icon16x16, FStyleColors::Warning));
+	Set("RevisionControl.NotAtHeadRevision", new CORE_IMAGE_BRUSH_SVG("Starship/SourceControl/RC_NewerVersion", CoreStyleConstants::Icon16x16, StatusNotAtHeadRevisionColor));
+	Set("RevisionControl.NotInDepot", new CORE_IMAGE_BRUSH_SVG("Starship/SourceControl/SCC_NotInDepot", CoreStyleConstants::Icon16x16, FStyleColors::Warning));
 	
-	Set("RevisionControl.Branched", new IMAGE_BRUSH_SVG("Starship/SourceControl/SCC_Branched", CoreStyleConstants::Icon16x16, BranchedColor));
+	Set("RevisionControl.Branched", new CORE_IMAGE_BRUSH_SVG("Starship/SourceControl/SCC_Branched", CoreStyleConstants::Icon16x16, BranchedColor));
 
-	Set("RevisionControl.Conflicted", new IMAGE_BRUSH_SVG("Starship/SourceControl/RC_Conflicted", CoreStyleConstants::Icon16x16, FStyleColors::Warning));
+	Set("RevisionControl.Conflicted", new CORE_IMAGE_BRUSH_SVG("Starship/SourceControl/RC_Conflicted", CoreStyleConstants::Icon16x16, FStyleColors::Warning));
 	
 	// Misc Icons
-	Set("RevisionControl.ChangelistsTab", new IMAGE_BRUSH_SVG("Starship/Common/check-circle", CoreStyleConstants::Icon16x16));
-	Set("RevisionControl.ConflictResolutionTab", new IMAGE_BRUSH_SVG("Starship/SourceControl/RC_Conflicted", CoreStyleConstants::Icon16x16));
-	Set("RevisionControl.SnapshotHistoryTab", new IMAGE_BRUSH_SVG("Starship/SourceControl/RC_Rewind", CoreStyleConstants::Icon16x16));
+	Set("RevisionControl.ChangelistsTab", new CORE_IMAGE_BRUSH_SVG("Starship/Common/check-circle", CoreStyleConstants::Icon16x16));
+	Set("RevisionControl.ConflictResolutionTab", new CORE_IMAGE_BRUSH_SVG("Starship/SourceControl/RC_Conflicted", CoreStyleConstants::Icon16x16));
+	Set("RevisionControl.SnapshotHistoryTab", new CORE_IMAGE_BRUSH_SVG("Starship/SourceControl/RC_Rewind", CoreStyleConstants::Icon16x16));
 	
-	Set("RevisionControl.StatusBar.AtLatestRevision", new IMAGE_BRUSH_SVG("Starship/SourceControl/RC_StatusRemoteUpToDate", CoreStyleConstants::Icon16x16));
-	Set("RevisionControl.StatusBar.NotAtLatestRevision", new IMAGE_BRUSH_SVG("Starship/SourceControl/RC_StatusRemoteDownload", CoreStyleConstants::Icon16x16, StatusNotAtHeadRevisionColor));
-	Set("RevisionControl.StatusBar.NoLocalChanges", new IMAGE_BRUSH_SVG("Starship/SourceControl/RC_StatusLocalUpToDate", CoreStyleConstants::Icon16x16));
-	Set("RevisionControl.StatusBar.HasLocalChanges", new IMAGE_BRUSH_SVG("Starship/SourceControl/RC_StatusLocalUpload", CoreStyleConstants::Icon16x16, FStyleColors::AccentBlue));
-	Set("RevisionControl.StatusBar.Conflicted", new IMAGE_BRUSH_SVG("Starship/SourceControl/RC_Conflicted", CoreStyleConstants::Icon16x16, FStyleColors::Warning));
-	Set("RevisionControl.StatusBar.Promote", new IMAGE_BRUSH_SVG("Starship/SourceControl/RC_Promote", CoreStyleConstants::Icon16x16));
-
-	Set("RevisionControl.ShowMenu.CheckedOut", new IMAGE_BRUSH_SVG("Starship/SourceControl/SCC_CheckedOut", CoreStyleConstants::Icon16x16));
-	Set("RevisionControl.ShowMenu.OpenForAdd", new IMAGE_BRUSH_SVG("Starship/SourceControl/RC_MarkedForAdd", CoreStyleConstants::Icon16x16));
-	Set("RevisionControl.ShowMenu.CheckedOutByOtherUser", new IMAGE_BRUSH_SVG("Starship/SourceControl/RC_CheckedOther", CoreStyleConstants::Icon16x16));
-	Set("RevisionControl.ShowMenu.NotAtHeadRevision", new IMAGE_BRUSH_SVG("Starship/SourceControl/RC_NewerVersion", CoreStyleConstants::Icon16x16));
+	Set("RevisionControl.StatusBar.AtLatestRevision", new CORE_IMAGE_BRUSH_SVG("Starship/SourceControl/RC_StatusRemoteUpToDate", CoreStyleConstants::Icon16x16));
+	Set("RevisionControl.StatusBar.NotAtLatestRevision", new CORE_IMAGE_BRUSH_SVG("Starship/SourceControl/RC_StatusRemoteDownload", CoreStyleConstants::Icon16x16, StatusNotAtHeadRevisionColor));
+	Set("RevisionControl.StatusBar.NoLocalChanges", new CORE_IMAGE_BRUSH_SVG("Starship/SourceControl/RC_StatusLocalUpToDate", CoreStyleConstants::Icon16x16));
+	Set("RevisionControl.StatusBar.HasLocalChanges", new CORE_IMAGE_BRUSH_SVG("Starship/SourceControl/RC_StatusLocalUpload", CoreStyleConstants::Icon16x16, FStyleColors::AccentBlue));
+	Set("RevisionControl.StatusBar.Conflicted", new CORE_IMAGE_BRUSH_SVG("Starship/SourceControl/RC_Conflicted", CoreStyleConstants::Icon16x16, FStyleColors::Warning));
+	Set("RevisionControl.StatusBar.Promote", new CORE_IMAGE_BRUSH_SVG("Starship/SourceControl/RC_Promote", CoreStyleConstants::Icon16x16));
+
+	Set("RevisionControl.ShowMenu.CheckedOut", new CORE_IMAGE_BRUSH_SVG("Starship/SourceControl/SCC_CheckedOut", CoreStyleConstants::Icon16x16));
+	Set("RevisionControl.ShowMenu.OpenForAdd", new CORE_IMAGE_BRUSH_SVG("Starship/SourceControl/RC_MarkedForAdd", CoreStyleConstants::Icon16x16));
+	Set("RevisionControl.ShowMenu.CheckedOutByOtherUser", new CORE_IMAGE_BRUSH_SVG("Starship/SourceControl/RC_CheckedOther", CoreStyleConstants::Icon16x16));
+	Set("RevisionControl.ShowMenu.NotAtHeadRevision", new CORE_IMAGE_BRUSH_SVG("Starship/SourceControl/RC_NewerVersion", CoreStyleConstants::Icon16x16));
 	
-	Set("RevisionControl.Promote.Large", new IMAGE_BRUSH_SVG("Starship/SourceControl/RC_Promote_Large", CoreStyleConstants::Icon32x32));
+	Set("RevisionControl.Promote.Large", new CORE_IMAGE_BRUSH_SVG("Starship/SourceControl/RC_Promote_Large", CoreStyleConstants::Icon32x32));
 
-	Set("RevisionControl.ConflictResolution.OpenExternal", new IMAGE_BRUSH_SVG("Starship/SourceControl/RC_ConflictResolution_OpenExternal", CoreStyleConstants::Icon16x16));
-	Set("RevisionControl.ConflictResolution.Clear", new IMAGE_BRUSH_SVG("Starship/SourceControl/RC_ConflictResolution_Clear", CoreStyleConstants::Icon16x16));
+	Set("RevisionControl.ConflictResolution.OpenExternal", new CORE_IMAGE_BRUSH_SVG("Starship/SourceControl/RC_ConflictResolution_OpenExternal", CoreStyleConstants::Icon16x16));
+	Set("RevisionControl.ConflictResolution.Clear", new CORE_IMAGE_BRUSH_SVG("Starship/SourceControl/RC_ConflictResolution_Clear", CoreStyleConstants::Icon16x16));
 
-	Set("RevisionControl.SnapshotHistory.Added", new IMAGE_BRUSH_SVG("Starship/SourceControl/RC_Added", CoreStyleConstants::Icon14x14, SnapshotHistoryAdded));
-	Set("RevisionControl.SnapshotHistory.Modified", new IMAGE_BRUSH_SVG("Starship/SourceControl/RC_Modified", CoreStyleConstants::Icon14x14, SnapshotHistoryModified));
-	Set("RevisionControl.SnapshotHistory.Removed", new IMAGE_BRUSH_SVG("Starship/SourceControl/RC_Removed", CoreStyleConstants::Icon14x14, SnapshotHistoryRemoved));
+	Set("RevisionControl.SnapshotHistory.Added", new CORE_IMAGE_BRUSH_SVG("Starship/SourceControl/RC_Added", CoreStyleConstants::Icon14x14, SnapshotHistoryAdded));
+	Set("RevisionControl.SnapshotHistory.Modified", new CORE_IMAGE_BRUSH_SVG("Starship/SourceControl/RC_Modified", CoreStyleConstants::Icon14x14, SnapshotHistoryModified));
+	Set("RevisionControl.SnapshotHistory.Removed", new CORE_IMAGE_BRUSH_SVG("Starship/SourceControl/RC_Removed", CoreStyleConstants::Icon14x14, SnapshotHistoryRemoved));
 
-	Set("RevisionControl.SnapshotHistory.ToolTip.Added", new IMAGE_BRUSH_SVG("Starship/SourceControl/RC_Added", CoreStyleConstants::Icon14x14));
-	Set("RevisionControl.SnapshotHistory.ToolTip.Modified", new IMAGE_BRUSH_SVG("Starship/SourceControl/RC_Modified", CoreStyleConstants::Icon14x14));
-	Set("RevisionControl.SnapshotHistory.ToolTip.Removed", new IMAGE_BRUSH_SVG("Starship/SourceControl/RC_Removed", CoreStyleConstants::Icon14x14));
+	Set("RevisionControl.SnapshotHistory.ToolTip.Added", new CORE_IMAGE_BRUSH_SVG("Starship/SourceControl/RC_Added", CoreStyleConstants::Icon14x14));
+	Set("RevisionControl.SnapshotHistory.ToolTip.Modified", new CORE_IMAGE_BRUSH_SVG("Starship/SourceControl/RC_Modified", CoreStyleConstants::Icon14x14));
+	Set("RevisionControl.SnapshotHistory.ToolTip.Removed", new CORE_IMAGE_BRUSH_SVG("Starship/SourceControl/RC_Removed", CoreStyleConstants::Icon14x14));
 
 	// Revision Control States
-	Set("RevisionControl.VerticalLine", new IMAGE_BRUSH_SVG("Starship/SourceControl/RC_VerticalLine", CoreStyleConstants::Icon26x26));
-	Set("RevisionControl.VerticalLineStart", new IMAGE_BRUSH_SVG("Starship/SourceControl/RC_VerticalLineStart", CoreStyleConstants::Icon26x26));
-	Set("RevisionControl.VerticalLineDashed", new IMAGE_BRUSH_SVG("Starship/SourceControl/RC_VerticalLineDashed", CoreStyleConstants::Icon26x26));
-	Set("RevisionControl.CheckCircleLine", new IMAGE_BRUSH_SVG("Starship/SourceControl/RC_CheckCircleLine", CoreStyleConstants::Icon26x26));
-	Set("RevisionControl.LineCircle", new IMAGE_BRUSH_SVG("Starship/SourceControl/RC_LineCircle", CoreStyleConstants::Icon26x26));
-	Set("RevisionControl.CheckInAvailable", new IMAGE_BRUSH_SVG("Starship/SourceControl/RC_CheckInAvailable", CoreStyleConstants::Icon26x26));
-	Set("RevisionControl.Rewound", new IMAGE_BRUSH_SVG("Starship/SourceControl/RC_Rewound", CoreStyleConstants::Icon26x26));
-	Set("RevisionControl.CheckInAvailableRewound", new IMAGE_BRUSH_SVG("Starship/SourceControl/RC_CheckInAvailableRewound", CoreStyleConstants::Icon26x26));
-	Set("RevisionControl.ConflictedState", new IMAGE_BRUSH_SVG("Starship/SourceControl/RC_ConflictedState", CoreStyleConstants::Icon26x26));
-
-	Set("RevisionControl.File", new IMAGE_BRUSH_SVG("Starship/SourceControl/RC_File", CoreStyleConstants::Icon6x8));
-	Set("RevisionControl.DiskSize", new IMAGE_BRUSH_SVG("Starship/SourceControl/RC_DiskSize", CoreStyleConstants::Icon8x8));
+	Set("RevisionControl.VerticalLine", new CORE_IMAGE_BRUSH_SVG("Starship/SourceControl/RC_VerticalLine", CoreStyleConstants::Icon26x26));
+	Set("RevisionControl.VerticalLineStart", new CORE_IMAGE_BRUSH_SVG("Starship/SourceControl/RC_VerticalLineStart", CoreStyleConstants::Icon26x26));
+	Set("RevisionControl.VerticalLineDashed", new CORE_IMAGE_BRUSH_SVG("Starship/SourceControl/RC_VerticalLineDashed", CoreStyleConstants::Icon26x26));
+	Set("RevisionControl.CheckCircleLine", new CORE_IMAGE_BRUSH_SVG("Starship/SourceControl/RC_CheckCircleLine", CoreStyleConstants::Icon26x26));
+	Set("RevisionControl.LineCircle", new CORE_IMAGE_BRUSH_SVG("Starship/SourceControl/RC_LineCircle", CoreStyleConstants::Icon26x26));
+	Set("RevisionControl.CheckInAvailable", new CORE_IMAGE_BRUSH_SVG("Starship/SourceControl/RC_CheckInAvailable", CoreStyleConstants::Icon26x26));
+	Set("RevisionControl.Rewound", new CORE_IMAGE_BRUSH_SVG("Starship/SourceControl/RC_Rewound", CoreStyleConstants::Icon26x26));
+	Set("RevisionControl.CheckInAvailableRewound", new CORE_IMAGE_BRUSH_SVG("Starship/SourceControl/RC_CheckInAvailableRewound", CoreStyleConstants::Icon26x26));
+	Set("RevisionControl.ConflictedState", new CORE_IMAGE_BRUSH_SVG("Starship/SourceControl/RC_ConflictedState", CoreStyleConstants::Icon26x26));
+
+	Set("RevisionControl.File", new CORE_IMAGE_BRUSH_SVG("Starship/SourceControl/RC_File", CoreStyleConstants::Icon6x8));
+	Set("RevisionControl.DiskSize", new CORE_IMAGE_BRUSH_SVG("Starship/SourceControl/RC_DiskSize", CoreStyleConstants::Icon8x8));
 }
 
 FDefaultRevisionControlStyle::~FDefaultRevisionControlStyle()
 {
-	FSlateStyleRegistry::UnRegisterSlateStyle(*this);
 }
 	
 const FName& FDefaultRevisionControlStyle::GetStyleSetName() const
diff --git a/Engine/Source/Developer/SourceControl/Public/RevisionControlStyle/RevisionControlStyle.h b/Engine/Source/Developer/SourceControl/Public/RevisionControlStyle/RevisionControlStyle.h
--- a/Engine/Source/Developer/SourceControl/Public/RevisionControlStyle/RevisionControlStyle.h
+++ b/Engine/Source/Developer/SourceControl/Public/RevisionControlStyle/RevisionControlStyle.h
@@ -14,6 +14,8 @@ class FRevisionControlStyleManager
 {
 public:
 
+	~FRevisionControlStyleManager();
+
 	/** reloads textures used by slate renderer */
 	static SOURCECONTROL_API void ReloadTextures();
 
@@ -47,8 +49,8 @@ class FDefaultRevisionControlStyle : public FSlateStyleSet
 {
 public:
 	
-	FDefaultRevisionControlStyle();
-	virtual ~FDefaultRevisionControlStyle() override;
+	SOURCECONTROL_API FDefaultRevisionControlStyle();
+	SOURCECONTROL_API virtual ~FDefaultRevisionControlStyle() override;
 	
 	virtual const FName& GetStyleSetName() const override;
 
