PR #13509: Adding a previous user property to FSourceControlState

diff --git a/Engine/Source/Developer/SourceControl/Public/SourceControlHelpers.h b/Engine/Source/Developer/SourceControl/Public/SourceControlHelpers.h
--- a/Engine/Source/Developer/SourceControl/Public/SourceControlHelpers.h
+++ b/Engine/Source/Developer/SourceControl/Public/SourceControlHelpers.h
@@ -92,6 +92,21 @@ public:
 			{
 				InStatus->IsCheckedOutOther(&CheckedOutOther);
 			}
+
+			//Go through the history to find the most recent user
+			const int32 HistorySize = InStatus->GetHistorySize();
+			for ( int32 Index = 0; Index < HistorySize; Index++ )
+			{
+				if ( const ISourceControlRevision* Revision = InStatus->GetHistoryItem( Index ).Get() )
+				{
+					const FString User = Revision->GetUserName();
+					if ( !User.IsEmpty() )
+					{
+						PreviousUser = User;
+						break;
+					}
+				}
+			}
 		}
 		else
 		{
@@ -190,6 +205,10 @@ public:
 	/** Determine if this file is modified by someone in another branch */
 	UPROPERTY(VisibleAnywhere, BlueprintReadOnly, Category = "Editor Scripting | Editor Revision Control Helpers", meta = (Keywords = "Source Control"))
 	bool bIsModifiedInOtherBranch = false;
+
+	/** The username of the previous user to modify this file */
+	UPROPERTY(VisibleAnywhere, BlueprintReadOnly, Category = "Editor Scripting | Editor Revision Control Helpers", meta = (Keywords = "Source Control"))
+	FString PreviousUser;
 };  // FSourceControlState
 
 
