Skein / URC2: SourceControl should batch `FUpdateStatus` queries into larger requests for URC where possible.

#rnx

diff --git a/Engine/Source/Developer/SourceControl/Private/SourceControlFileStatusMonitor.cpp b/Engine/Source/Developer/SourceControl/Private/SourceControlFileStatusMonitor.cpp
--- a/Engine/Source/Developer/SourceControl/Private/SourceControlFileStatusMonitor.cpp
+++ b/Engine/Source/Developer/SourceControl/Private/SourceControlFileStatusMonitor.cpp
@@ -297,7 +297,7 @@ bool FSourceControlFileStatusMonitor::Tick(float DeltaTime)
 		});
 	}
 
-	RequestedStatusFiles.Reserve(MaxFileNumPerRequestPolicy);
+	RequestedStatusFiles.Reserve(RequestedStatusFiles.Num() + NewFiles.Num() + RefreshedFiles.Num());
 	while (RequestedStatusFiles.Num() < MaxFileNumPerRequestPolicy)
 	{
 		if (NewFiles.Num() > 0)
diff --git a/Engine/Source/Developer/SourceControl/Private/SourceControlModule.cpp b/Engine/Source/Developer/SourceControl/Private/SourceControlModule.cpp
--- a/Engine/Source/Developer/SourceControl/Private/SourceControlModule.cpp
+++ b/Engine/Source/Developer/SourceControl/Private/SourceControlModule.cpp
@@ -36,9 +36,6 @@ static const FName SourceControlFeatureName("SourceControl");
 
 namespace SourceControlConstants
 {
-	/** The maximum number of file/directory status requests we should dispatch in a tick */
-	const int32 MaxStatusDispatchesPerTick = 64;
-
 	/** The interval at which we refresh a file's state */
 	const FTimespan StateRefreshInterval = FTimespan::FromMinutes(5.0);
 }
@@ -295,12 +292,15 @@ void FSourceControlModule::Tick()
 			// check for any pending dispatches
 			if (PendingStatusUpdateFiles.Num() > 0)
 			{
+				// get batch size
+				const int32 MaxStatusDispatchesPerRequest = GetSourceControlFileStatusMonitor().GetMaxFilePerRequestPolicy();
+
 				// grab a batch of files
 				TArray<FString> FilesToDispatch;
-				FilesToDispatch.Reserve(SourceControlConstants::MaxStatusDispatchesPerTick);
+				FilesToDispatch.Reserve(MaxStatusDispatchesPerRequest);
 				for (const FString& Filename : PendingStatusUpdateFiles)
 				{
-					if(FilesToDispatch.Num() >= SourceControlConstants::MaxStatusDispatchesPerTick)
+					if(FilesToDispatch.Num() >= MaxStatusDispatchesPerRequest)
 					{
 						break;
 					}
