SourceControl: Time-slice the FUpdateStatus + SetUpdateHistory(true) operation over multiple requests.

When deleting a large prefab (say 2000 individual files) and opening the 'SUBMIT FILES' dialog, the AssetDataCache receives 2000 requests for asset data so their pretty names can be displayed in the dialog (as the pretty name is contained within the asset, and the actual asset is no longer on disk). This FUpdateStatus + SetUpdateHistory(true) operation can be quite long running so timeslicing it over multiple requests prevents the user from encountering large SCC delays on providers that only support 1 simultaneous operation at a time.

#rnx

diff --git a/Engine/Source/Developer/SourceControl/Private/SourceControlAssetDataCache.cpp b/Engine/Source/Developer/SourceControl/Private/SourceControlAssetDataCache.cpp
--- a/Engine/Source/Developer/SourceControl/Private/SourceControlAssetDataCache.cpp
+++ b/Engine/Source/Developer/SourceControl/Private/SourceControlAssetDataCache.cpp
@@ -59,6 +59,7 @@ void FSourceControlAssetDataCache::Startup()
 #endif
 
 	bIsSourceControlDialogShown = false;
+	bIsUpdateHistoryInFlight = false;
 }
 
 void FSourceControlAssetDataCache::Shutdown()
@@ -167,29 +168,48 @@ void FSourceControlAssetDataCache::OnSourceControlDialogClosed()
 
 void FSourceControlAssetDataCache::GetFileHistory()
 {
+	if (bIsUpdateHistoryInFlight)
+	{
+		return;
+	}
+
 	if (FileHistoryToUpdate.Num() == 0)
 	{
 		return;
 	}
 
-	ISourceControlProvider& SourceControlProvider = ISourceControlModule::Get().GetProvider();
+	const int32 MaxStatusDispatchesPerRequest = 100; // Use a conservative number, SetUpdateHistory(true) could be expensive.
+	const int32 NumItemsAvailable = FileHistoryToUpdate.Num();
+	const int32 NumItemsToDispatch = FMath::Min(MaxStatusDispatchesPerRequest, NumItemsAvailable);
+
+	TArray<FString> FilesToDispatch;
+	FilesToDispatch.Reserve(NumItemsToDispatch);
+	for (int32 Index = 0; Index < NumItemsToDispatch; ++Index)
+	{
+		FilesToDispatch.Add(FileHistoryToUpdate[Index]);
+	}
+	FileHistoryToUpdate.RemoveAt(0, NumItemsToDispatch);
+
+	bIsUpdateHistoryInFlight = true;
+
 	TSharedRef<FUpdateStatus, ESPMode::ThreadSafe> UpdateStatusOperation = ISourceControlOperation::Create<FUpdateStatus>();
 	UpdateStatusOperation->SetUpdateHistory(true);
 	UpdateStatusOperation->SetQuiet(true);
 
-	SourceControlProvider.Execute(UpdateStatusOperation,
-								  FileHistoryToUpdate,
-								  EConcurrency::Asynchronous,
-								  FSourceControlOperationComplete::CreateLambda([FileHistoryToUpdateCopy = FileHistoryToUpdate](const FSourceControlOperationRef& Operation, ECommandResult::Type InResult)
-	{
-		ISourceControlModule::Get().GetAssetDataCache().OnUpdateHistoryComplete(FileHistoryToUpdateCopy, Operation, InResult);
-	}));
-
-	FileHistoryToUpdate.Reset();
+	ISourceControlModule::Get().GetProvider().Execute(UpdateStatusOperation, FilesToDispatch, EConcurrency::Asynchronous,
+		FSourceControlOperationComplete::CreateLambda(
+			[FilesToDispatch](const FSourceControlOperationRef& Operation, ECommandResult::Type InResult)
+			{
+				ISourceControlModule::Get().GetAssetDataCache().OnUpdateHistoryComplete(FilesToDispatch, Operation, InResult);
+			}
+		)
+	);
 }
 
 void FSourceControlAssetDataCache::OnUpdateHistoryComplete(const TArray<FString>& InUpdatedFiles, const FSourceControlOperationRef& InOperation, ECommandResult::Type InResult)
 {
+	bIsUpdateHistoryInFlight = false;
+
 	if (InResult == ECommandResult::Failed)
 	{
 		UE_LOG(LogSourceControl, Warning, TEXT("UpdateHistory Failed."));
diff --git a/Engine/Source/Developer/SourceControl/Public/SourceControlAssetDataCache.h b/Engine/Source/Developer/SourceControl/Public/SourceControlAssetDataCache.h
--- a/Engine/Source/Developer/SourceControl/Public/SourceControlAssetDataCache.h
+++ b/Engine/Source/Developer/SourceControl/Public/SourceControlAssetDataCache.h
@@ -116,4 +116,5 @@ private:
 	FDelegateHandle										AssetUpdatedDelegateHandle;
 	std::atomic<uint32>									CurrentAsyncTask = 0;
 	bool												bIsSourceControlDialogShown;
+	bool												bIsUpdateHistoryInFlight;
 };
