Skein / URC2: Tweaks to improve clarity of pre-emptive conflicts warning.

#rnx

diff --git a/Engine/Source/Developer/SourceControl/Private/RevisionControlStyle/RevisionControlStyle.cpp b/Engine/Source/Developer/SourceControl/Private/RevisionControlStyle/RevisionControlStyle.cpp
--- a/Engine/Source/Developer/SourceControl/Private/RevisionControlStyle/RevisionControlStyle.cpp
+++ b/Engine/Source/Developer/SourceControl/Private/RevisionControlStyle/RevisionControlStyle.cpp
@@ -149,7 +149,8 @@ FDefaultRevisionControlStyle::FDefaultRevisionControlStyle() : FSlateStyleSet(St
 	Set("RevisionControl.StatusBar.NotAtLatestRevision", new CORE_IMAGE_BRUSH_SVG("Starship/SourceControl/RC_StatusRemoteDownload", CoreStyleConstants::Icon16x16, StatusNotAtHeadRevisionColor));
 	Set("RevisionControl.StatusBar.NoLocalChanges", new CORE_IMAGE_BRUSH_SVG("Starship/SourceControl/RC_StatusLocalUpToDate", CoreStyleConstants::Icon16x16));
 	Set("RevisionControl.StatusBar.HasLocalChanges", new CORE_IMAGE_BRUSH_SVG("Starship/SourceControl/RC_StatusLocalUpload", CoreStyleConstants::Icon16x16, FStyleColors::AccentBlue));
-	Set("RevisionControl.StatusBar.Conflicted", new CORE_IMAGE_BRUSH_SVG("Starship/SourceControl/RC_Conflicted", CoreStyleConstants::Icon16x16, FStyleColors::Warning));
+	Set("RevisionControl.StatusBar.Conflict", new CORE_IMAGE_BRUSH_SVG("Starship/SourceControl/RC_Conflicted", CoreStyleConstants::Icon16x16, FStyleColors::Warning));
+	Set("RevisionControl.StatusBar.ConflictUpcoming", new CORE_IMAGE_BRUSH_SVG("Starship/Common/alert-triangle", CoreStyleConstants::Icon16x16, FStyleColors::Warning));
 	Set("RevisionControl.StatusBar.Promote", new CORE_IMAGE_BRUSH_SVG("Starship/SourceControl/RC_Promote", CoreStyleConstants::Icon16x16));
 
 	Set("RevisionControl.ShowMenu.CheckedOut", new CORE_IMAGE_BRUSH_SVG("Starship/SourceControl/SCC_CheckedOut", CoreStyleConstants::Icon16x16));
diff --git a/Engine/Source/Developer/SourceControl/Private/SSourceControlControls.cpp b/Engine/Source/Developer/SourceControl/Private/SSourceControlControls.cpp
--- a/Engine/Source/Developer/SourceControl/Private/SSourceControlControls.cpp
+++ b/Engine/Source/Developer/SourceControl/Private/SSourceControlControls.cpp
@@ -305,7 +305,7 @@ FText SSourceControlControls::GetSourceControlSyncStatusToolTipText()
 	}
 	if (GetNumConflictsUpcoming() > 0)
 	{
-		return LOCTEXT("SyncLatestButtonNotAtHeadTooltipTextConflictUpcoming", "Some of your local changes conflict with the latest snapshot of the project. Click here to review these conflicts.");
+		return LOCTEXT("SyncLatestButtonNotAtHeadTooltipTextConflictUpcoming", "Some of your local changes conflict with later snapshots on the project. Sync Latest or Check-in Changes to review and address these issues.");
 	}
 	if (HasSourceControlChangesToSync())
 	{
@@ -317,7 +317,8 @@ FText SSourceControlControls::GetSourceControlSyncStatusToolTipText()
 
 const FSlateBrush* SSourceControlControls::GetSourceControlSyncStatusIcon()
 {
-	static const FSlateBrush* ConflictBrush = FRevisionControlStyleManager::Get().GetBrush("RevisionControl.StatusBar.Conflicted");
+	static const FSlateBrush* ConflictBrush = FRevisionControlStyleManager::Get().GetBrush("RevisionControl.StatusBar.Conflict");
+	static const FSlateBrush* ConflictUpcomingBrush = FRevisionControlStyleManager::Get().GetBrush("RevisionControl.StatusBar.ConflictUpcoming");
 	static const FSlateBrush* AtHeadBrush = FRevisionControlStyleManager::Get().GetBrush("RevisionControl.StatusBar.AtLatestRevision");
 	static const FSlateBrush* NotAtHeadBrush = FRevisionControlStyleManager::Get().GetBrush("RevisionControl.StatusBar.NotAtLatestRevision");
 
@@ -327,7 +328,7 @@ const FSlateBrush* SSourceControlControls::GetSourceControlSyncStatusIcon()
 	}
 	if (GetNumConflictsUpcoming() > 0)
 	{
-		return ConflictBrush;
+		return ConflictUpcomingBrush;
 	}
 	if (HasSourceControlChangesToSync())
 	{
@@ -444,7 +445,7 @@ FText SSourceControlControls::GetSourceControlCheckInStatusToolTipText()
 	}
 	if (GetNumConflictsUpcoming() > 0)
 	{
-		return LOCTEXT("CheckInButtonChangesTooltipTextConflictUpcoming", "Some of your local changes conflict with the latest snapshot of the project. Click here to review these conflicts.");
+		return LOCTEXT("CheckInButtonChangesTooltipTextConflictUpcoming", "Some of your local changes conflict with later snapshots on the project. Sync Latest or Check-in Changes to review and address these issues.");
 	}
 	if (HasSourceControlChangesToCheckIn())
 	{
@@ -456,7 +457,8 @@ FText SSourceControlControls::GetSourceControlCheckInStatusToolTipText()
 
 const FSlateBrush* SSourceControlControls::GetSourceControlCheckInStatusIcon()
 {
-	static const FSlateBrush* ConflictBrush = FRevisionControlStyleManager::Get().GetBrush("RevisionControl.StatusBar.Conflicted");
+	static const FSlateBrush* ConflictBrush = FRevisionControlStyleManager::Get().GetBrush("RevisionControl.StatusBar.Conflict");
+	static const FSlateBrush* ConflictUpcomingBrush = FRevisionControlStyleManager::Get().GetBrush("RevisionControl.StatusBar.ConflictUpcoming");
 	static const FSlateBrush* NoLocalChangesBrush = FRevisionControlStyleManager::Get().GetBrush("RevisionControl.StatusBar.NoLocalChanges");
 	static const FSlateBrush* HasLocalChangesBrush = FRevisionControlStyleManager::Get().GetBrush("RevisionControl.StatusBar.HasLocalChanges");
 
@@ -466,7 +468,7 @@ const FSlateBrush* SSourceControlControls::GetSourceControlCheckInStatusIcon()
 	}
 	if (GetNumConflictsUpcoming() > 0)
 	{
-		return ConflictBrush;
+		return ConflictUpcomingBrush;
 	}
 	if (HasSourceControlChangesToCheckIn())
 	{
