Add Branch Explorer icon

diff --git a/Engine/Source/Developer/SourceControl/Private/RevisionControlStyle/RevisionControlStyle.cpp b/Engine/Source/Developer/SourceControl/Private/RevisionControlStyle/RevisionControlStyle.cpp
--- a/Engine/Source/Developer/SourceControl/Private/RevisionControlStyle/RevisionControlStyle.cpp
+++ b/Engine/Source/Developer/SourceControl/Private/RevisionControlStyle/RevisionControlStyle.cpp
@@ -110,6 +110,7 @@ FDefaultRevisionControlStyle::FDefaultRevisionControlStyle() : FSlateStyleSet(St
 	Set("RevisionControl.Actions.Refresh", new CORE_IMAGE_BRUSH_SVG("Starship/Common/Update", CoreStyleConstants::Icon16x16));
 	Set("RevisionControl.Actions.ChangeSettings", new CORE_IMAGE_BRUSH_SVG("Starship/SourceControl/icon_SCC_Change_Source_Control_Settings", CoreStyleConstants::Icon16x16));
 	Set("RevisionControl.Actions.Rewind", new CORE_IMAGE_BRUSH_SVG("Starship/SourceControl/RC_Rewind", CoreStyleConstants::Icon16x16));
+	Set("RevisionControl.Actions.BranchExplorer", new CORE_IMAGE_BRUSH_SVG("Starship/SourceControl/RC_BranchExplorer", CoreStyleConstants::Icon16x16));
 
 	// Icons representing the various revision control states
 
@@ -144,6 +145,7 @@ FDefaultRevisionControlStyle::FDefaultRevisionControlStyle() : FSlateStyleSet(St
 	Set("RevisionControl.ChangelistsTab", new CORE_IMAGE_BRUSH_SVG("Starship/Common/check-circle", CoreStyleConstants::Icon16x16));
 	Set("RevisionControl.ConflictResolutionTab", new CORE_IMAGE_BRUSH_SVG("Starship/SourceControl/RC_Conflicted", CoreStyleConstants::Icon16x16));
 	Set("RevisionControl.SnapshotHistoryTab", new CORE_IMAGE_BRUSH_SVG("Starship/SourceControl/RC_Rewind", CoreStyleConstants::Icon16x16));
+	Set("RevisionControl.BranchExplorerTab", new CORE_IMAGE_BRUSH_SVG("Starship/SourceControl/RC_BranchExplorer", CoreStyleConstants::Icon16x16));
 	
 	Set("RevisionControl.StatusBar.AtLatestRevision", new CORE_IMAGE_BRUSH_SVG("Starship/SourceControl/RC_StatusRemoteUpToDate", CoreStyleConstants::Icon16x16));
 	Set("RevisionControl.StatusBar.NotAtLatestRevision", new CORE_IMAGE_BRUSH_SVG("Starship/SourceControl/RC_StatusRemoteDownload", CoreStyleConstants::Icon16x16, StatusNotAtHeadRevisionColor));
