Update source control styles

#rnx

diff --git a/Engine/Source/Developer/SourceControl/Private/RevisionControlStyle/RevisionControlStyle.cpp b/Engine/Source/Developer/SourceControl/Private/RevisionControlStyle/RevisionControlStyle.cpp
--- a/Engine/Source/Developer/SourceControl/Private/RevisionControlStyle/RevisionControlStyle.cpp
+++ b/Engine/Source/Developer/SourceControl/Private/RevisionControlStyle/RevisionControlStyle.cpp
@@ -110,7 +110,6 @@ FDefaultRevisionControlStyle::FDefaultRevisionControlStyle() : FSlateStyleSet(St
 	Set("RevisionControl.Actions.Refresh", new CORE_IMAGE_BRUSH_SVG("Starship/Common/Update", CoreStyleConstants::Icon16x16));
 	Set("RevisionControl.Actions.ChangeSettings", new CORE_IMAGE_BRUSH_SVG("Starship/SourceControl/icon_SCC_Change_Source_Control_Settings", CoreStyleConstants::Icon16x16));
 	Set("RevisionControl.Actions.Rewind", new CORE_IMAGE_BRUSH_SVG("Starship/SourceControl/RC_Rewind", CoreStyleConstants::Icon16x16));
-	Set("RevisionControl.Actions.BranchExplorer", new CORE_IMAGE_BRUSH_SVG("Starship/SourceControl/RC_BranchExplorer", CoreStyleConstants::Icon16x16));
 
 	// Icons representing the various revision control states
 
diff --git a/Engine/Source/Editor/StatusBar/Private/SourceControlMenuHelpers.cpp b/Engine/Source/Editor/StatusBar/Private/SourceControlMenuHelpers.cpp
--- a/Engine/Source/Editor/StatusBar/Private/SourceControlMenuHelpers.cpp
+++ b/Engine/Source/Editor/StatusBar/Private/SourceControlMenuHelpers.cpp
@@ -332,7 +332,7 @@ TSharedRef<SWidget> FSourceControlMenuHelpers::GenerateSourceControlMenuContent(
 		FSourceControlCommands::Get().ViewBranchExplorer,
 		TAttribute<FText>(),
 		TAttribute<FText>(),
-		FSlateIcon(FRevisionControlStyleManager::GetStyleSetName(), "RevisionControl.Actions.BranchExplorer")
+		FSlateIcon(FRevisionControlStyleManager::GetStyleSetName(), "RevisionControl.BranchExplorerTab")
 	);
 
 	Section.AddMenuEntry(
