Modify ISourceControlModule submit validation delegates.
This is so that we can route all of the validation messages to the delegate caller instead of just the warnings and errors.
Previous delegate set/get pure virtual functions are marked deprecated and made final to catch all override classes.

diff --git a/Engine/Plugins/Editor/DataValidation/Source/DataValidation/Private/EditorValidatorSubsystem.cpp b/Engine/Plugins/Editor/DataValidation/Source/DataValidation/Private/EditorValidatorSubsystem.cpp
--- a/Engine/Plugins/Editor/DataValidation/Source/DataValidation/Private/EditorValidatorSubsystem.cpp
+++ b/Engine/Plugins/Editor/DataValidation/Source/DataValidation/Private/EditorValidatorSubsystem.cpp
@@ -118,8 +118,8 @@ void UEditorValidatorSubsystem::Initialize(FSubsystemCollectionBase& Collection)
 	}
 
 	// Register to SCC pre-submit callback
-	ISourceControlModule::Get().RegisterPreSubmitChangelistValidation(FSourceControlPreSubmitChangelistValidationDelegate::CreateUObject(this, &UEditorValidatorSubsystem::ValidateChangelistPreSubmit));
-	ISourceControlModule::Get().RegisterPreSubmitFilesValidation(FSourceControlPreSubmitFilesValidationDelegate::CreateUObject(this, &UEditorValidatorSubsystem::ValidateFilesPreSubmit));
+	ISourceControlModule::Get().RegisterPreSubmitChangelistValidationEx(FSourceControlPreSubmitChangelistValidationExDelegate::CreateUObject(this, &UEditorValidatorSubsystem::ValidateChangelistPreSubmit));
+	ISourceControlModule::Get().RegisterPreSubmitFilesValidationEx(FSourceControlPreSubmitFilesValidationExDelegate::CreateUObject(this, &UEditorValidatorSubsystem::ValidateFilesPreSubmit));
 }
 
 bool UEditorValidatorSubsystem::ShouldValidateAsset(
@@ -966,9 +966,7 @@ void UEditorValidatorSubsystem::ValidateAllSavedPackages()
 
 void UEditorValidatorSubsystem::ValidateChangelistPreSubmit(
 	FSourceControlChangelistPtr InChangelist,
-	EDataValidationResult& 		OutResult,
-	TArray<FText>& 				OutValidationErrors,
-	TArray<FText>& 				OutValidationWarnings) const
+	FSourceControlValidationDelegateResult& OutResult) const
 {
 	check(InChangelist.IsValid());
 
@@ -985,22 +983,25 @@ void UEditorValidatorSubsystem::ValidateChangelistPreSubmit(
 	Settings.bShowIfNoFailures = false;
 
 	FValidateAssetsResults Results;
-	OutResult = ValidateChangelist(Changelist, Settings, Results);
+	OutResult.Result = ValidateChangelist(Changelist, Settings, Results);
 	
 	for (const TPair<FString, FValidateAssetsDetails>& Pair : Results.AssetsDetails)
 	{
 		const FValidateAssetsDetails& Details = Pair.Value;
 
-		OutValidationWarnings.Append(Details.ValidationWarnings);
-		OutValidationErrors.Append(Details.ValidationErrors);
+		OutResult.Warnings.Append(Details.ValidationWarnings);
+		OutResult.Errors.Append(Details.ValidationErrors);
+		
+		for (TSharedRef<FTokenizedMessage> const& Message : Details.ValidationMessages)
+		{
+			OutResult.AllMessages.Add(Message->ToText());
+		}
 	}
 }
 
 void UEditorValidatorSubsystem::ValidateFilesPreSubmit(
 	TArray<FSourceControlStateRef>& InStates,
-	EDataValidationResult& OutResult,
-	TArray<FText>& OutValidationErrors,
-	TArray<FText>& OutValidationWarnings) const
+	FSourceControlValidationDelegateResult& OutResult) const
 {
 	// Create temporary changelist object to do most of the heavy lifting
 	UDataValidationChangelist* Changelist = NewObject<UDataValidationChangelist>();
@@ -1015,14 +1016,22 @@ void UEditorValidatorSubsystem::ValidateFilesPreSubmit(
 	Settings.bShowIfNoFailures = false;
 
 	FValidateAssetsResults Results;
-	OutResult = ValidateChangelist(Changelist, Settings, Results);
+	OutResult.Result = ValidateChangelist(Changelist, Settings, Results);
 
 	for (const TPair<FString, FValidateAssetsDetails>& Pair : Results.AssetsDetails)
 	{
 		const FValidateAssetsDetails& Details = Pair.Value;
 
-		OutValidationWarnings.Append(Details.ValidationWarnings);
-		OutValidationErrors.Append(Details.ValidationErrors);
+		OutResult.Warnings.Append(Details.ValidationWarnings);
+		OutResult.Errors.Append(Details.ValidationErrors);
+
+		for (TSharedRef<FTokenizedMessage> const& Message : Details.ValidationMessages)
+		{
+			if (Message->GetSeverity() == EMessageSeverity::Info)
+			{
+				OutResult.AllMessages.Add(Message->ToText());
+			}
+		}
 	}
 }
 
diff --git a/Engine/Plugins/Editor/DataValidation/Source/DataValidation/Public/EditorValidatorSubsystem.h b/Engine/Plugins/Editor/DataValidation/Source/DataValidation/Public/EditorValidatorSubsystem.h
--- a/Engine/Plugins/Editor/DataValidation/Source/DataValidation/Public/EditorValidatorSubsystem.h
+++ b/Engine/Plugins/Editor/DataValidation/Source/DataValidation/Public/EditorValidatorSubsystem.h
@@ -21,6 +21,7 @@ class FMessageLog;
 struct FARFilter;
 struct FAssetData;
 class ISourceControlChangelist;
+struct FSourceControlValidationDelegateResult;
 typedef TSharedPtr<class ISourceControlChangelist, ESPMode::ThreadSafe> FSourceControlChangelistPtr;
 typedef TSharedRef<class ISourceControlState, ESPMode::ThreadSafe> FSourceControlStateRef;
 
@@ -437,12 +438,12 @@ protected:
 	/**
 	 * Validates changelist before SCC submit operation
 	 */
-	UE_API void ValidateChangelistPreSubmit(FSourceControlChangelistPtr Changelist, EDataValidationResult& OutResult, TArray<FText>& ValidationErrors, TArray<FText>& ValidationWarnings) const;
+	UE_API void ValidateChangelistPreSubmit(FSourceControlChangelistPtr Changelist, FSourceControlValidationDelegateResult& OutResult) const;
 
 	/**
 	 * Validates files before SCC submit operation
 	 */
-	UE_API void ValidateFilesPreSubmit(TArray<FSourceControlStateRef>& States, EDataValidationResult& OutResult, TArray<FText>& ValidationErrors, TArray<FText>& ValidationWarnings) const;
+	UE_API void ValidateFilesPreSubmit(TArray<FSourceControlStateRef>& States, FSourceControlValidationDelegateResult& OutResult) const;
 
 	/**
 	 * Handle native modules being loaded or unloaded (@see NativeModulesPendingLoad and NativeModulesPendingUnload)
diff --git a/Engine/Plugins/Experimental/SubmitToolEditor/Source/SubmitToolEditor/Private/SubmitToolEditor.cpp b/Engine/Plugins/Experimental/SubmitToolEditor/Source/SubmitToolEditor/Private/SubmitToolEditor.cpp
--- a/Engine/Plugins/Experimental/SubmitToolEditor/Source/SubmitToolEditor/Private/SubmitToolEditor.cpp
+++ b/Engine/Plugins/Experimental/SubmitToolEditor/Source/SubmitToolEditor/Private/SubmitToolEditor.cpp
@@ -156,7 +156,7 @@ FSubmitOverrideReply FSubmitToolEditorModule::OnSubmitToolOverrideCallback(const
 	{
 		if (NormalizedPath.EndsWith(TEXT(".exe"), ESearchCase::IgnoreCase))
 		{
-			NormalizedPath += NormalizedPath.Replace(TEXT(".exe"), TEXT(""), ESearchCase::IgnoreCase);
+			NormalizedPath = NormalizedPath.Replace(TEXT(".exe"), TEXT(""), ESearchCase::IgnoreCase);
 		}
 	}
 
@@ -456,17 +456,15 @@ bool FSubmitToolEditorModule::GetChangelistValidationResult(ISourceControlProvid
 			break;
 		}
 	}
-	if (ChangelistToProcess->IsValid())
+	if (ChangelistToProcess.IsValid())
 	{
-		FSourceControlPreSubmitChangelistValidationDelegate ValidationDelegate = ISourceControlModule::Get().GetRegisteredPreSubmitChangelistValidation();
+		FSourceControlPreSubmitChangelistValidationExDelegate ValidationDelegate = ISourceControlModule::Get().GetRegisteredPreSubmitChangelistValidationEx();
 
-		EDataValidationResult ValidationResult = EDataValidationResult::NotValidated;
-		TArray<FText> ValidationErrors;
-		TArray<FText> ValidationWarnings;
-
-		ValidationDelegate.ExecuteIfBound(*(ChangelistToProcess), ValidationResult, ValidationErrors, ValidationWarnings);
+		FSourceControlValidationDelegateResult ValidationResult;
+		
+		ValidationDelegate.ExecuteIfBound(*(ChangelistToProcess), ValidationResult);
 
-		return (ValidationResult == EDataValidationResult::Valid);
+		return (ValidationResult.Result == EDataValidationResult::Valid);
 	}
 	else
 	{
diff --git a/Engine/Source/Developer/SourceControl/Private/SourceControlModule.cpp b/Engine/Source/Developer/SourceControl/Private/SourceControlModule.cpp
--- a/Engine/Source/Developer/SourceControl/Private/SourceControlModule.cpp
+++ b/Engine/Source/Developer/SourceControl/Private/SourceControlModule.cpp
@@ -569,32 +569,32 @@ void FSourceControlModule::UnregisterProviderChanged(FDelegateHandle Handle)
 	OnSourceControlProviderChanged.Remove(Handle);
 }
 
-void FSourceControlModule::RegisterPreSubmitChangelistValidation(const FSourceControlPreSubmitChangelistValidationDelegate& PreSubmitChangelistValidationDelegate)
+void FSourceControlModule::RegisterPreSubmitChangelistValidationEx(const FSourceControlPreSubmitChangelistValidationExDelegate& PreSubmitChangelistValidationDelegate)
 {
 	OnSourceControlPreSubmitChangelistValidation = PreSubmitChangelistValidationDelegate;
 }
 
 void FSourceControlModule::UnregisterPreSubmitChangelistValidation()
 {
-	OnSourceControlPreSubmitChangelistValidation = FSourceControlPreSubmitChangelistValidationDelegate();
+	OnSourceControlPreSubmitChangelistValidation = FSourceControlPreSubmitChangelistValidationExDelegate();
 }
 
-FSourceControlPreSubmitChangelistValidationDelegate FSourceControlModule::GetRegisteredPreSubmitChangelistValidation()
+FSourceControlPreSubmitChangelistValidationExDelegate FSourceControlModule::GetRegisteredPreSubmitChangelistValidationEx()
 {
 	return OnSourceControlPreSubmitChangelistValidation;
 }
 
-void FSourceControlModule::RegisterPreSubmitFilesValidation(const FSourceControlPreSubmitFilesValidationDelegate& PreSubmitFilesValidationDelegate)
+void FSourceControlModule::RegisterPreSubmitFilesValidationEx(const FSourceControlPreSubmitFilesValidationExDelegate& PreSubmitFilesValidationDelegate)
 {
 	OnSourceControlPreSubmitFilesValidation = PreSubmitFilesValidationDelegate;
 }
 
 void FSourceControlModule::UnregisterPreSubmitFilesValidation()
 {
-	OnSourceControlPreSubmitFilesValidation = FSourceControlPreSubmitFilesValidationDelegate();
+	OnSourceControlPreSubmitFilesValidation = FSourceControlPreSubmitFilesValidationExDelegate();
 }
 
-FSourceControlPreSubmitFilesValidationDelegate FSourceControlModule::GetRegisteredPreSubmitFilesValidation()
+FSourceControlPreSubmitFilesValidationExDelegate FSourceControlModule::GetRegisteredPreSubmitFilesValidationEx()
 {
 	return OnSourceControlPreSubmitFilesValidation;
 }
diff --git a/Engine/Source/Developer/SourceControl/Private/SourceControlModule.h b/Engine/Source/Developer/SourceControl/Private/SourceControlModule.h
--- a/Engine/Source/Developer/SourceControl/Private/SourceControlModule.h
+++ b/Engine/Source/Developer/SourceControl/Private/SourceControlModule.h
@@ -46,12 +46,14 @@ public:
 	virtual FSourceControlProviderChanging& GetSourceControlProviderChanging() override;
 	virtual FDelegateHandle RegisterProviderChanged(const FSourceControlProviderChanged::FDelegate& SourceControlProviderChanged) override;
 	virtual void UnregisterProviderChanged(FDelegateHandle Handle) override;
-	virtual void RegisterPreSubmitChangelistValidation(const FSourceControlPreSubmitChangelistValidationDelegate& PreSubmitChangelistValidationDelegate) override;
+	
+	virtual void RegisterPreSubmitChangelistValidationEx(const FSourceControlPreSubmitChangelistValidationExDelegate& PreSubmitChangelistValidationDelegate) override;
 	virtual void UnregisterPreSubmitChangelistValidation() override;
-	virtual void RegisterPreSubmitFilesValidation(const FSourceControlPreSubmitFilesValidationDelegate& PreSubmitFilesValidationDelegate) override;
+	virtual FSourceControlPreSubmitChangelistValidationExDelegate GetRegisteredPreSubmitChangelistValidationEx() override;
+	
+	virtual void RegisterPreSubmitFilesValidationEx(const FSourceControlPreSubmitFilesValidationExDelegate& PreSubmitFilesValidationDelegate) override;
 	virtual void UnregisterPreSubmitFilesValidation() override;
-	virtual FSourceControlPreSubmitChangelistValidationDelegate GetRegisteredPreSubmitChangelistValidation() override;
-	virtual FSourceControlPreSubmitFilesValidationDelegate GetRegisteredPreSubmitFilesValidation() override;
+	virtual FSourceControlPreSubmitFilesValidationExDelegate GetRegisteredPreSubmitFilesValidationEx() override;
 	
 	virtual FDelegateHandle RegisterPreSubmitFinalize(const FSourceControlPreSubmitFinalizeDelegate::FDelegate& Delegate) override;
 	virtual void UnregisterPreSubmitFinalize(FDelegateHandle Handle) override;
@@ -167,10 +169,10 @@ private:
 	FSourceControlProviderChanged OnSourceControlProviderChanged;
 
 	/** To call when doing pre-submit changelist validation */
-	FSourceControlPreSubmitChangelistValidationDelegate OnSourceControlPreSubmitChangelistValidation;
+	FSourceControlPreSubmitChangelistValidationExDelegate OnSourceControlPreSubmitChangelistValidation;
 
 	/** To call when doing pre-submit files validation */
-	FSourceControlPreSubmitFilesValidationDelegate OnSourceControlPreSubmitFilesValidation;
+	FSourceControlPreSubmitFilesValidationExDelegate OnSourceControlPreSubmitFilesValidation;
 	
 	/** To be called right before files are submitted, allowing for additional last minute validation. @see FSourceControlPreSubmitFinalizeDelegate */
 	FSourceControlPreSubmitFinalizeDelegate OnPresubmitFinalize;
diff --git a/Engine/Source/Developer/SourceControl/Public/ISourceControlModule.h b/Engine/Source/Developer/SourceControl/Public/ISourceControlModule.h
--- a/Engine/Source/Developer/SourceControl/Public/ISourceControlModule.h
+++ b/Engine/Source/Developer/SourceControl/Public/ISourceControlModule.h
@@ -58,11 +58,28 @@ DECLARE_DELEGATE_RetVal(bool, FSourceControlProviderChanging);
 /** Delegate called when the active source control provider is changed */
 DECLARE_MULTICAST_DELEGATE_TwoParams( FSourceControlProviderChanged, ISourceControlProvider& /*OldProvider*/, ISourceControlProvider& /*NewProvider*/ );
 
-/** Delegate called on pre-submit for data validation */
+/** Deprecated delegates called on pre-submit for data validation */
 DECLARE_DELEGATE_FourParams(FSourceControlPreSubmitDataValidationDelegate, FSourceControlChangelistPtr /*Changelist*/, EDataValidationResult& /*Result*/, TArray<FText>& /*ValidationErrors*/, TArray<FText>& /*ValidationWarnings*/);
 DECLARE_DELEGATE_FourParams(FSourceControlPreSubmitChangelistValidationDelegate, FSourceControlChangelistPtr /*Changelist*/, EDataValidationResult& /*Result*/, TArray<FText>& /*ValidationErrors*/, TArray<FText>& /*ValidationWarnings*/);
 DECLARE_DELEGATE_FourParams(FSourceControlPreSubmitFilesValidationDelegate, TArray<FSourceControlStateRef>& /*States*/, EDataValidationResult& /*Result*/, TArray<FText>& /*ValidationErrors*/, TArray<FText>& /*ValidationWarnings*/);
 
+/** Structure to hold result of pre-submit data validation. */
+struct FSourceControlValidationDelegateResult
+{
+	/** Validation result. */
+	EDataValidationResult Result = EDataValidationResult::NotValidated;
+	/** Error messages output during validation. */
+	TArray<FText> Errors;
+	/** Warning messages output during validation. */
+	TArray<FText> Warnings;
+	/** All messages (including error and warnings output during validation. */
+	TArray<FText> AllMessages;
+};
+
+/** Delegates called on pre-submit for data validation. */
+DECLARE_DELEGATE_TwoParams(FSourceControlPreSubmitChangelistValidationExDelegate, FSourceControlChangelistPtr /*Changelist*/, FSourceControlValidationDelegateResult& /*Result*/);
+DECLARE_DELEGATE_TwoParams(FSourceControlPreSubmitFilesValidationExDelegate, TArray<FSourceControlStateRef>& /*States*/, FSourceControlValidationDelegateResult& /*Result*/);
+
 /** Delegate that returns directory information on custom projects under source control (first entry is the main project) */
 DECLARE_DELEGATE_RetVal(TArray<FSourceControlProjectInfo>, FSourceControlCustomProjectsDelegate);
 
@@ -248,27 +265,26 @@ public:
 	 */
 	virtual void UnregisterProviderChanged(FDelegateHandle Handle) = 0;
 
-	/**
-	 * Register a delegate to be called to validate asset changes before submitting changes
-	 */
-	UE_DEPRECATED(5.7, "Use RegisterPreSubmitChangelistValidation instead.")
-	virtual void RegisterPreSubmitDataValidation(const FSourceControlPreSubmitDataValidationDelegate& PreSubmitDataValidationDelegate) { RegisterPreSubmitChangelistValidation(PreSubmitDataValidationDelegate); };
+	// Deprecating pure virtuals using final keyword to catch derived implementations.
+	UE_DEPRECATED(5.7, "Use RegisterPreSubmitChangelistValidationEx instead.")
+	virtual void RegisterPreSubmitDataValidation(const FSourceControlPreSubmitDataValidationDelegate& PreSubmitDataValidationDelegate) final {};
+	UE_DEPRECATED(5.8, "Use RegisterPreSubmitChangelistValidationEx instead.")
+	virtual void RegisterPreSubmitChangelistValidation(const FSourceControlPreSubmitChangelistValidationDelegate& PreSubmitChangelistValidationDelegate) final {};
+	UE_DEPRECATED(5.8, "Use RegisterPreSubmitFilesValidationEx instead.")
+	virtual void RegisterPreSubmitFilesValidation(const FSourceControlPreSubmitFilesValidationDelegate& PreSubmitFilesValidationDelegate) final {};
+	UE_DEPRECATED(5.7, "Use UnregisterPreSubmitChangelistValidation instead.")
+	virtual void UnregisterPreSubmitDataValidation() final {};
+	UE_DEPRECATED(5.7, "Use GetRegisteredPreSubmitChangelistValidation instead.")
+	virtual FSourceControlPreSubmitDataValidationDelegate GetRegisteredPreSubmitDataValidation() final { return FSourceControlPreSubmitDataValidationDelegate(); };
+	UE_DEPRECATED(5.8, "Use GetRegisteredPreSubmitChangelistValidationEx instead.")
+	virtual FSourceControlPreSubmitChangelistValidationDelegate GetRegisteredPreSubmitChangelistValidation() final { return FSourceControlPreSubmitChangelistValidationDelegate(); };
+	UE_DEPRECATED(5.8, "Use GetRegisteredPreSubmitFilesValidationEx instead.")
+	virtual FSourceControlPreSubmitFilesValidationDelegate GetRegisteredPreSubmitFilesValidation() final { return FSourceControlPreSubmitFilesValidationDelegate(); };
 
 	/**
 	 * Register a delegate to be called to validate asset changes by changelist before submitting changes
 	 */
-	virtual void RegisterPreSubmitChangelistValidation(const FSourceControlPreSubmitChangelistValidationDelegate& PreSubmitChangelistValidationDelegate) = 0;
-
-	/**
-	 * Register a delegate to be called to validate asset changes by files before submitting changes
-	 */
-	virtual void RegisterPreSubmitFilesValidation(const FSourceControlPreSubmitFilesValidationDelegate& PreSubmitFilesValidationDelegate) = 0;
-
-	/**
-	 * Unregister a delegate called before submitting changes
-	 */
-	UE_DEPRECATED(5.7, "Use UnregisterPreSubmitChangelistValidation instead.")
-	virtual void UnregisterPreSubmitDataValidation() { UnregisterPreSubmitChangelistValidation(); };
+	virtual void RegisterPreSubmitChangelistValidationEx(const FSourceControlPreSubmitChangelistValidationExDelegate& PreSubmitChangelistValidationDelegate) = 0;
 
 	/**
 	 * Unregister a delegate called before submitting changes
@@ -276,25 +292,24 @@ public:
 	virtual void UnregisterPreSubmitChangelistValidation() = 0;
 
 	/**
-	 * Unregister a delegate called before submitting changes
+	 * Gets currently registered delegates for pre-submit changelist validation
 	 */
-	virtual void UnregisterPreSubmitFilesValidation() = 0;
+	virtual FSourceControlPreSubmitChangelistValidationExDelegate GetRegisteredPreSubmitChangelistValidationEx() = 0;
 
 	/**
-	 * Gets currently registered delegates for pre-submit data validation
+	 * Register a delegate to be called to validate asset changes by files before submitting changes
 	 */
-	UE_DEPRECATED(5.7, "Use GetRegisteredPreSubmitChangelistValidation instead.")
-	virtual FSourceControlPreSubmitDataValidationDelegate GetRegisteredPreSubmitDataValidation() { return GetRegisteredPreSubmitChangelistValidation(); };
+	virtual void RegisterPreSubmitFilesValidationEx(const FSourceControlPreSubmitFilesValidationExDelegate& PreSubmitFilesValidationDelegate) = 0;
 
 	/**
-	 * Gets currently registered delegates for pre-submit changelist validation
+	 * Unregister a delegate called before submitting changes
 	 */
-	virtual FSourceControlPreSubmitChangelistValidationDelegate GetRegisteredPreSubmitChangelistValidation() = 0;
+	virtual void UnregisterPreSubmitFilesValidation() = 0;
 
 	/**
 	 * Gets currently registered delegates for pre-submit files validation
 	 */
-	virtual FSourceControlPreSubmitFilesValidationDelegate GetRegisteredPreSubmitFilesValidation() = 0;
+	virtual FSourceControlPreSubmitFilesValidationExDelegate GetRegisteredPreSubmitFilesValidationEx() = 0;
 
 	/** 
 	 * Register a delegate that is invokes right before files are submitted to source control. @see FSourceControlPreSubmitFinalizeDelegate
diff --git a/Engine/Source/Editor/SourceControlWindows/Private/SSourceControlChangelists.cpp b/Engine/Source/Editor/SourceControlWindows/Private/SSourceControlChangelists.cpp
--- a/Engine/Source/Editor/SourceControlWindows/Private/SSourceControlChangelists.cpp
+++ b/Engine/Source/Editor/SourceControlWindows/Private/SSourceControlChangelists.cpp
@@ -2337,25 +2337,23 @@ void SSourceControlChangelistsWidget::OnDeleteShelvedFiles()
 
 static bool GetChangelistValidationResult(FSourceControlChangelistPtr InChangelist, FString& OutValidationTitleText, FString& OutValidationWarningsText, FString& OutValidationErrorsText)
 {
-	FSourceControlPreSubmitChangelistValidationDelegate ValidationDelegate = ISourceControlModule::Get().GetRegisteredPreSubmitChangelistValidation();
+	FSourceControlPreSubmitChangelistValidationExDelegate ValidationDelegate = ISourceControlModule::Get().GetRegisteredPreSubmitChangelistValidationEx();
 
-	EDataValidationResult ValidationResult = EDataValidationResult::NotValidated;
-	TArray<FText> ValidationErrors;
-	TArray<FText> ValidationWarnings;
+	FSourceControlValidationDelegateResult ValidationResult;
 
 	bool bValidationResult = true;
 
-	if (ValidationDelegate.ExecuteIfBound(InChangelist, ValidationResult, ValidationErrors, ValidationWarnings))
+	if (ValidationDelegate.ExecuteIfBound(InChangelist, ValidationResult))
 	{
 		EMessageSeverity::Type MessageSeverity = EMessageSeverity::Info;
 
-		if (ValidationResult == EDataValidationResult::Invalid || ValidationErrors.Num() > 0)
+		if (ValidationResult.Result == EDataValidationResult::Invalid || ValidationResult.Errors.Num() > 0)
 		{
 			OutValidationTitleText = LOCTEXT("SourceControl.Submit.ChangelistValidationError", "Changelist validation failed!").ToString();
 			bValidationResult = false;
 			MessageSeverity = EMessageSeverity::Error;
 		}
-		else if (ValidationResult == EDataValidationResult::NotValidated || ValidationWarnings.Num() > 0)
+		else if (ValidationResult.Result == EDataValidationResult::NotValidated || ValidationResult.Warnings.Num() > 0)
 		{
 			OutValidationTitleText = LOCTEXT("SourceControl.Submit.ChangelistValidationWarning", "Changelist validation has warnings!").ToString();
 			MessageSeverity = EMessageSeverity::Warning;
@@ -2398,8 +2396,8 @@ static bool GetChangelistValidationResult(FSourceControlChangelistPtr InChangeli
 				}
 			};
 
-		AppendInfo(ValidationErrors, LOCTEXT("ValidationErrorsHeader", "Encountered {Count} {Count}|plural(one=error,other=errors):"), LOCTEXT("ValidationErrorsOverflow", "See log for complete list of errors"), OutValidationErrorsText);
-		AppendInfo(ValidationWarnings, LOCTEXT("ValidationWarningsHeader", "Encountered {Count} {Count}|plural(one=warning,other=warnings):"), LOCTEXT("ValidationWarningsOverflow", "See log for complete list of warnings"), OutValidationWarningsText);
+		AppendInfo(ValidationResult.Errors, LOCTEXT("ValidationErrorsHeader", "Encountered {Count} {Count}|plural(one=error,other=errors):"), LOCTEXT("ValidationErrorsOverflow", "See log for complete list of errors"), OutValidationErrorsText);
+		AppendInfo(ValidationResult.Warnings, LOCTEXT("ValidationWarningsHeader", "Encountered {Count} {Count}|plural(one=warning,other=warnings):"), LOCTEXT("ValidationWarningsOverflow", "See log for complete list of warnings"), OutValidationWarningsText);
 	}
 
 	return bValidationResult;
diff --git a/Engine/Source/Editor/SourceControlWindows/Private/SourceControlWindows.cpp b/Engine/Source/Editor/SourceControlWindows/Private/SourceControlWindows.cpp
--- a/Engine/Source/Editor/SourceControlWindows/Private/SourceControlWindows.cpp
+++ b/Engine/Source/Editor/SourceControlWindows/Private/SourceControlWindows.cpp
@@ -90,25 +90,23 @@ static bool SaveDirtyPackages(bool bUseDialog)
 
 static bool GetFilesValidationResult(TArray<FSourceControlStateRef>& InStates, FString& OutValidationTitleText, FString& OutValidationWarningsText, FString& OutValidationErrorsText)
 {
-	FSourceControlPreSubmitFilesValidationDelegate ValidationDelegate = ISourceControlModule::Get().GetRegisteredPreSubmitFilesValidation();
+	FSourceControlPreSubmitFilesValidationExDelegate ValidationDelegate = ISourceControlModule::Get().GetRegisteredPreSubmitFilesValidationEx();
 
-	EDataValidationResult ValidationResult = EDataValidationResult::NotValidated;
-	TArray<FText> ValidationErrors;
-	TArray<FText> ValidationWarnings;
+	FSourceControlValidationDelegateResult ValidationResult;
 
 	bool bValidationResult = true;
 
-	if (ValidationDelegate.ExecuteIfBound(InStates, ValidationResult, ValidationErrors, ValidationWarnings))
+	if (ValidationDelegate.ExecuteIfBound(InStates, ValidationResult))
 	{
 		EMessageSeverity::Type MessageSeverity = EMessageSeverity::Info;
 
-		if (ValidationResult == EDataValidationResult::Invalid || ValidationErrors.Num() > 0)
+		if (ValidationResult.Result == EDataValidationResult::Invalid || ValidationResult.Errors.Num() > 0)
 		{
 			OutValidationTitleText = LOCTEXT("SourceControl.Submit.FilesValidationError", "File(s) validation failed!").ToString();
 			bValidationResult = false;
 			MessageSeverity = EMessageSeverity::Error;
 		}
-		else if (ValidationResult == EDataValidationResult::NotValidated || ValidationWarnings.Num() > 0)
+		else if (ValidationResult.Result == EDataValidationResult::NotValidated || ValidationResult.Warnings.Num() > 0)
 		{
 			OutValidationTitleText = LOCTEXT("SourceControl.Submit.FilesValidationWarning", "File(s) validation has warnings!").ToString();
 			MessageSeverity = EMessageSeverity::Warning;
@@ -151,8 +149,8 @@ static bool GetFilesValidationResult(TArray<FSourceControlStateRef>& InStates, F
 				}
 			};
 
-		AppendInfo(ValidationErrors, LOCTEXT("ValidationErrorsHeader", "Encountered {Count} {Count}|plural(one=error,other=errors):"), LOCTEXT("ValidationErrorsOverflow", "See log for complete list of errors"), OutValidationErrorsText);
-		AppendInfo(ValidationWarnings, LOCTEXT("ValidationWarningsHeader", "Encountered {Count} {Count}|plural(one=warning,other=warnings):"), LOCTEXT("ValidationWarningsOverflow", "See log for complete list of warnings"), OutValidationWarningsText);
+		AppendInfo(ValidationResult.Errors, LOCTEXT("ValidationErrorsHeader", "Encountered {Count} {Count}|plural(one=error,other=errors):"), LOCTEXT("ValidationErrorsOverflow", "See log for complete list of errors"), OutValidationErrorsText);
+		AppendInfo(ValidationResult.Warnings, LOCTEXT("ValidationWarningsHeader", "Encountered {Count} {Count}|plural(one=warning,other=warnings):"), LOCTEXT("ValidationWarningsOverflow", "See log for complete list of warnings"), OutValidationWarningsText);
 	}
 
 	return bValidationResult;
