SourceControl: Added a time threshold before triggering an FUpdateStatus after repeated calls to QueueStatusUpdate.

Context: While scrolling through the SceneOutliner, it calls ISourceControlModule::QueueStatusUpdate as new items come into view. In a test project with 300+ items, slowly scrolling the SceneOutliner from top to bottom, this resulted in 80+ separate FUpdateStatus operations as they were getting spawned every tick.

After introducing this threshold, it resulted in 1 FUpdateStatus operation on all 300+ items at once.

#rnx

diff --git a/Engine/Source/Developer/SourceControl/Private/SourceControlCVars.cpp b/Engine/Source/Developer/SourceControl/Private/SourceControlCVars.cpp
--- a/Engine/Source/Developer/SourceControl/Private/SourceControlCVars.cpp
+++ b/Engine/Source/Developer/SourceControl/Private/SourceControlCVars.cpp
@@ -28,4 +28,9 @@ namespace SourceControlCVars
 		TEXT("SourceControl.EnableWorldReload"),
 		false,
 		TEXT("Allows a SourceControl operation to perform a world reload."));
+
+	TAutoConsoleVariable<int32> CVarSourceControlIntervalBatchStatusUpdate(
+		TEXT("SourceControl.IntervalBatchStatusUpdate"),
+		1,
+		TEXT("Controls interval that allows QueueStatusUpdate requests to be batched with others."));
 }
\ No newline at end of file
diff --git a/Engine/Source/Developer/SourceControl/Private/SourceControlCVars.h b/Engine/Source/Developer/SourceControl/Private/SourceControlCVars.h
--- a/Engine/Source/Developer/SourceControl/Private/SourceControlCVars.h
+++ b/Engine/Source/Developer/SourceControl/Private/SourceControlCVars.h
@@ -11,4 +11,6 @@ namespace SourceControlCVars
 	extern TAutoConsoleVariable<bool> CVarSourceControlEnableRevertFromSubmitWidget;
 	extern TAutoConsoleVariable<bool> CVarSourceControlEnableRevertUnsaved;
 	extern TAutoConsoleVariable<bool> CVarSourceControlEnableLoginDialogModal;
+
+	extern TAutoConsoleVariable<int32> CVarSourceControlIntervalBatchStatusUpdate;
 }
\ No newline at end of file
diff --git a/Engine/Source/Developer/SourceControl/Private/SourceControlModule.cpp b/Engine/Source/Developer/SourceControl/Private/SourceControlModule.cpp
--- a/Engine/Source/Developer/SourceControl/Private/SourceControlModule.cpp
+++ b/Engine/Source/Developer/SourceControl/Private/SourceControlModule.cpp
@@ -289,8 +289,12 @@ void FSourceControlModule::Tick()
 		// don't allow background status updates when temporarily disabled for login
 		if (TemporarilyDisabledCount == 0)
 		{
+			// determine delay before processing pending requests
+			const int32 Delay = SourceControlCVars::CVarSourceControlIntervalBatchStatusUpdate.GetValueOnGameThread();
+
 			// check for any pending dispatches
-			if (PendingStatusUpdateFiles.Num() > 0)
+			const FDateTime Now = FDateTime::UtcNow();
+			if ((PendingStatusUpdateFiles.Num() > 0) && (PendingStatusUpdateQueuedTime + FTimespan::FromSeconds(Delay) < Now))
 			{
 				// get batch size
 				const int32 MaxStatusDispatchesPerRequest = GetSourceControlFileStatusMonitor().GetMaxFilePerRequestPolicy();
@@ -361,6 +365,7 @@ void FSourceControlModule::QueueStatusUpdate(const FString& InFilename)
 			if (TimeSinceLastUpdate > SourceControlConstants::StateRefreshInterval)
 			{
 				PendingStatusUpdateFiles.AddUnique(InFilename);
+				PendingStatusUpdateQueuedTime = FDateTime::UtcNow();
 			}
 		}
 	}
diff --git a/Engine/Source/Developer/SourceControl/Private/SourceControlModule.h b/Engine/Source/Developer/SourceControl/Private/SourceControlModule.h
--- a/Engine/Source/Developer/SourceControl/Private/SourceControlModule.h
+++ b/Engine/Source/Developer/SourceControl/Private/SourceControlModule.h
@@ -156,6 +156,9 @@ private:
 	/** Files pending a status update */
 	TArray<FString> PendingStatusUpdateFiles;
 
+	/** Time last status update got queued */
+	FDateTime PendingStatusUpdateQueuedTime;
+
 	/** Flag to disable source control - used temporarily when login is in progress */
 	uint8 TemporarilyDisabledCount;
 
