SourceControl: Fixed SourceControlFileStatusMonitor triggering an FUpdateStatus after duplicating an actor.

At that point, the actor merely exists in memory - it has not been saved to disk yet. So there is no SCC state to query.

#rnx

#Synced-CL 47898077

diff --git a/Engine/Source/Developer/SourceControl/Private/SourceControlFileStatusMonitor.cpp b/Engine/Source/Developer/SourceControl/Private/SourceControlFileStatusMonitor.cpp
--- a/Engine/Source/Developer/SourceControl/Private/SourceControlFileStatusMonitor.cpp
+++ b/Engine/Source/Developer/SourceControl/Private/SourceControlFileStatusMonitor.cpp
@@ -15,8 +15,36 @@
 #include "Framework/Application/SlateApplication.h"
 #endif //#if SOURCE_CONTROL_WITH_SLATE
 
+#if WITH_EDITORONLY_DATA
+#include "Misc/PackageName.h"
+#include "UObject/Package.h"
+#endif // WITH_EDITORONLY_DATA
+
+namespace
+{
+
+static bool CanQuerySourceControl(const FString& AbsPathname)
+{
+#if WITH_EDITORONLY_DATA
+	FString PackageName;
+	if (FPackageName::TryConvertFilenameToLongPackageName(AbsPathname, PackageName))
+	{
+		if (UPackage* Package = FindPackage(nullptr, *PackageName))
+		{
+			// There's nothing to query if the package didn't come from disk or hasn't yet been saved to disk.
+			return Package->GetHasBeenEndLoaded();
+		}
+	}
+#endif // WITH_EDITORONLY_DATA
+
+	return true;
+}
+
+}
+
 FSourceControlFileStatusMonitor::FSourceControlFileStatusMonitor()
 : ProbationPeriodPolicy(FTimespan::FromSeconds(1))
+, RemainingPeriodPolicy(FTimespan::FromSeconds(5))
 , RefreshPeriodPolicy(FTimespan::FromMinutes(5))
 {
 	TickerHandle = FTSTicker::GetCoreTicker().AddTicker(FTickerDelegate::CreateRaw(this, &FSourceControlFileStatusMonitor::Tick));
@@ -253,20 +281,28 @@ bool FSourceControlFileStatusMonitor::Tick(float DeltaTime)
 		return true;
 	}
 
-	double NowSecs = FPlatformTime::Seconds();
+	const double NowSecs = FPlatformTime::Seconds();
 
 	// Throttle the source control status check when no new files need to be checked. Don't overload the source control server with too many requests.
-	if (NewAddedFileCount == 0 && NowSecs - OldestFileStatusTimeSecs < RefreshPeriodPolicy.GetTotalSeconds())
+	if (NewAddedFileCount == 0 && NowSecs - OldestFileStatusTimeSecs < RefreshPeriodPolicy.GetTotalSeconds() && LastRemainingFileTimeSecs == 0.f)
 	{
 		return true; // Nothing to do this time around.
 	}
 
+	// Throttle the checks when remaining files are left.
+	if (NewAddedFileCount == 0 && NowSecs - LastRemainingFileTimeSecs < RemainingPeriodPolicy.GetTotalSeconds() && LastRemainingFileTimeSecs != 0.f)
+	{
+		return true; // Delay the retry.
+	}
+
 	// Throttle the checks when new files are added. Batch new files edited close in time to be more efficient.
 	if (NowSecs - LastAddedFileTimeSecs < ProbationPeriodPolicy.GetTotalSeconds() && NewAddedFileCount < MaxFileNumPerRequestPolicy)
 	{
 		return true; // Delay the request, give chance to capture more files.
 	}
 
+	TRACE_CPUPROFILER_EVENT_SCOPE(FSourceControlFileStatusMonitor::Tick);
+
 	TArray<const TPair<FString, TSharedPtr<FSourceControlFileStatus>>*> NewFiles;
 	NewFiles.Reserve(NewAddedFileCount);
 
@@ -291,23 +327,36 @@ bool FSourceControlFileStatusMonitor::Tick(float DeltaTime)
 	{
 		// Get the status of all newly added files and refresh the ones that were less recently updated.
 		RefreshedFiles.Sort([](const TPair<FString, TSharedPtr<FSourceControlFileStatus>>& Lhs, const TPair<FString, TSharedPtr<FSourceControlFileStatus>>& Rhs)
-		{
-			// Sort ascending as we are going to use Last()/Pop() later, so the oldest must be at the end.
-			return Lhs.Value->LastStatusCheckTimestampSecs > Rhs.Value->LastStatusCheckTimestampSecs;
-		});
+			{
+				// Sort ascending as we are going to use Last()/Pop() later, so the oldest must be at the end.
+				return Lhs.Value->LastStatusCheckTimestampSecs > Rhs.Value->LastStatusCheckTimestampSecs;
+			});
 	}
 
+	uint32 NewFilesRemaining = NewFiles.Num();
+	uint32 RefreshedFilesRemaining = RefreshedFiles.Num();
+
 	RequestedStatusFiles.Reserve(RequestedStatusFiles.Num() + NewFiles.Num() + RefreshedFiles.Num());
 	while (RequestedStatusFiles.Num() < MaxFileNumPerRequestPolicy)
 	{
 		if (NewFiles.Num() > 0)
 		{
-			RequestedStatusFiles.Emplace(NewFiles.Last()->Key);
+			const FString& AbsPathname = NewFiles.Last()->Key;
+			if (CanQuerySourceControl(AbsPathname))
+			{
+				RequestedStatusFiles.Emplace(AbsPathname);
+				NewFilesRemaining--;
+			}
 			NewFiles.Pop(EAllowShrinking::No);
 		}
 		else if (RefreshedFiles.Num())
 		{
-			RequestedStatusFiles.Emplace(RefreshedFiles.Last()->Key);
+			const FString& AbsPathname = RefreshedFiles.Last()->Key;
+			if (CanQuerySourceControl(AbsPathname))
+			{
+				RequestedStatusFiles.Emplace(AbsPathname);
+				RefreshedFilesRemaining--;
+			}
 			RefreshedFiles.Pop(EAllowShrinking::No);
 		}
 		else
@@ -318,6 +367,7 @@ bool FSourceControlFileStatusMonitor::Tick(float DeltaTime)
 
 	// The remaining number of new files.
 	NewAddedFileCount = NewFiles.Num();
+	LastRemainingFileTimeSecs = NewFilesRemaining > 0 ? FPlatformTime::Seconds() : 0.f;
 
 	if (RequestedStatusFiles.Num())
 	{
@@ -338,7 +388,7 @@ void FSourceControlFileStatusMonitor::OnSourceControlStatusUpdate(const TSharedR
 
 	ensure(IsInGameThread()); // Concurrency issues if the callback in invoked from a background thread.
 
-	double NowSecs = FPlatformTime::Seconds();
+	const double NowSecs = FPlatformTime::Seconds();
 
 	ON_SCOPE_EXIT
 	{
diff --git a/Engine/Source/Developer/SourceControl/Public/SourceControlFileStatusMonitor.h b/Engine/Source/Developer/SourceControl/Public/SourceControlFileStatusMonitor.h
--- a/Engine/Source/Developer/SourceControl/Public/SourceControlFileStatusMonitor.h
+++ b/Engine/Source/Developer/SourceControl/Public/SourceControlFileStatusMonitor.h
@@ -175,6 +175,9 @@ private:
 	/** The newly added files probation period. Time left to batch several files into the same requests. */
 	FTimespan ProbationPeriodPolicy;
 
+	/** The remaining files retry period. Time left until we retry a file's CanQuerySourceControl status. */
+	FTimespan RemainingPeriodPolicy;
+
 	/** The status refresh period. Time left until we re-request a file status. */
 	FTimespan RefreshPeriodPolicy;
 
@@ -187,9 +190,12 @@ private:
 	/** Last time a request to the source control provider was issued. */
 	double LastSourceControlCheckSecs = 0.0;
 
-	/** Last time a file was assed to the list of monitored files. */
+	/** Last time a file was added to the list of monitored files. */
 	double LastAddedFileTimeSecs = 0.0;
 
+	/** Last time some files were left remaining after inspection. */
+	double LastRemainingFileTimeSecs = 0.0;
+
 	/** The oldest status update time. */
 	double OldestFileStatusTimeSecs = 0.0;
 
