Update revision control style

#rnx
#Synced-CL 48697675

diff --git a/Engine/Source/Developer/SourceControl/Private/RevisionControlStyle/RevisionControlStyle.cpp b/Engine/Source/Developer/SourceControl/Private/RevisionControlStyle/RevisionControlStyle.cpp
--- a/Engine/Source/Developer/SourceControl/Private/RevisionControlStyle/RevisionControlStyle.cpp
+++ b/Engine/Source/Developer/SourceControl/Private/RevisionControlStyle/RevisionControlStyle.cpp
@@ -173,6 +173,7 @@ FDefaultRevisionControlStyle::FDefaultRevisionControlStyle() : FSlateStyleSet(St
 	Set("RevisionControl.SnapshotHistory.ToolTip.Removed", new CORE_IMAGE_BRUSH_SVG("Starship/SourceControl/RC_Removed", CoreStyleConstants::Icon14x14));
 
 	Set("RevisionControl.BranchExplorer.BranchList.Branch", new CORE_IMAGE_BRUSH_SVG("Starship/SourceControl/RC_Branch", CoreStyleConstants::Icon16x16));
+	Set("RevisionControl.BranchExplorer.BranchList.PersonalBranch", new CORE_IMAGE_BRUSH_SVG("Starship/Common/Developer", CoreStyleConstants::Icon16x16));
 	Set("RevisionControl.BranchExplorer.Merge", new CORE_IMAGE_BRUSH_SVG("Starship/SourceControl/RC_Merged", CoreStyleConstants::Icon16x16));
 	Set("RevisionControl.BranchExplorer.NewBranch", new CORE_IMAGE_BRUSH_SVG("Starship/SourceControl/RC_NewBranch", CoreStyleConstants::Icon16x16));
 
