Skein / URC2: Fixed 'Restore as Latest' being active while there are conflicts to resolve.

#rnx
#Synced-CL 48787719

diff --git a/Engine/Source/Developer/SourceControl/Private/SSourceControlControls.cpp b/Engine/Source/Developer/SourceControl/Private/SSourceControlControls.cpp
--- a/Engine/Source/Developer/SourceControl/Private/SSourceControlControls.cpp
+++ b/Engine/Source/Developer/SourceControl/Private/SSourceControlControls.cpp
@@ -301,7 +301,7 @@ FText SSourceControlControls::GetSourceControlSyncStatusToolTipText()
 {
 	if (GetNumConflictsRemaining() > 0)
 	{
-		return LOCTEXT("SyncLatestButtonNotAtHeadTooltipTextConflict", "Some of your local changes conflict with the latest snapshot of the project. Click here to review these conflicts.");
+		return LOCTEXT("SyncLatestButtonNotAtHeadTooltipTextConflict", "Some of your local changes are still in conflict. Click here to review these conflicts.");
 	}
 	if (GetNumConflictsUpcoming() > 0)
 	{
@@ -441,7 +441,7 @@ FText SSourceControlControls::GetSourceControlCheckInStatusToolTipText()
 {
 	if (GetNumConflictsRemaining() > 0)
 	{
-		return LOCTEXT("CheckInButtonChangesTooltipTextConflict", "Some of your local changes conflict with the latest snapshot of the project. Click here to review these conflicts.");
+		return LOCTEXT("CheckInButtonChangesTooltipTextConflict", "Some of your local changes are still in conflict. Click here to review these conflicts.");
 	}
 	if (GetNumConflictsUpcoming() > 0)
 	{
@@ -534,19 +534,42 @@ FText SSourceControlControls::GetSourceControlRestoreAsLatestText()
 
 FText SSourceControlControls::GetSourceControlRestoreAsLatestToolTipText()
 {
+	if (GetNumConflictsRemaining() > 0)
+	{
+		return LOCTEXT("RestoreAsLatestTooltipTextConflict", "Some of your local changes are still in conflict. Click here to review these conflicts.");
+	}
+	
 	return LOCTEXT("RestoreAsLatestTooltipText", "Restore this snapshot to be the latest version of the project for all team members.");
 }
 
 const FSlateBrush* SSourceControlControls::GetSourceControlRestoreAsLatestStatusIcon()
 {
-	return FRevisionControlStyleManager::Get().GetBrush("RevisionControl.StatusBar.Promote");
+	static const FSlateBrush* ConflictBrush = FRevisionControlStyleManager::Get().GetBrush("RevisionControl.StatusBar.Conflict");
+	static const FSlateBrush* PromoteBrush = FRevisionControlStyleManager::Get().GetBrush("RevisionControl.StatusBar.Promote");
+
+	if (GetNumConflictsRemaining() > 0)
+	{
+		return ConflictBrush;
+	}
+
+	return PromoteBrush;
 }
 
 FReply SSourceControlControls::OnSourceControlRestoreAsLatestClicked()
 {
-	if (OnRestoreAsLatestClicked.IsBound())
+	if (GetNumConflictsRemaining() > 0)
 	{
-		OnRestoreAsLatestClicked.Execute();
+		if (IConsoleObject* CObj = IConsoleManager::Get().FindConsoleObject(TEXT("UnrealRevisionControl.FocusConflictResolution")))
+		{
+			CObj->AsCommand()->Execute(/*Args=*/TArray<FString>(), /*InWorld=*/nullptr, *GLog);
+		}
+	}
+	else
+	{
+		if (OnRestoreAsLatestClicked.IsBound())
+		{
+			OnRestoreAsLatestClicked.Execute();
+		}
 	}
 
 	return FReply::Handled();
