Skein / URC2: Assorted UI changes.

#rnx
#Synced-CL 48897982

diff --git a/Engine/Source/Developer/SourceControl/Private/SSourceControlControls.cpp b/Engine/Source/Developer/SourceControl/Private/SSourceControlControls.cpp
--- a/Engine/Source/Developer/SourceControl/Private/SSourceControlControls.cpp
+++ b/Engine/Source/Developer/SourceControl/Private/SSourceControlControls.cpp
@@ -86,12 +86,11 @@ void SSourceControlControls::Construct(const FArguments& InArgs)
 	[
 		SNew(SHorizontalBox)
 		+ SHorizontalBox::Slot() // Check In Changes Button
-		.VAlign(VAlign_Center)
 		.Padding(FMargin(0.0f, 0.0f, 4.0f, 0.0f))
 		.AutoWidth()
 		[
 			SNew(SButton)
-			.ButtonStyle(&FAppStyle::Get().GetWidgetStyle<FButtonStyle>("StatusBar.StatusBarButton"))
+			.ButtonStyle(FAppStyle::Get(), "SimpleButton")
 			.ToolTipText_Static(&SSourceControlControls::GetSourceControlCheckInStatusToolTipText)
 			.Visibility_Static(&SSourceControlControls::GetSourceControlCheckInStatusVisibility)
 			.IsEnabled_Static(&SSourceControlControls::IsSourceControlCheckInEnabled)
@@ -118,12 +117,11 @@ void SSourceControlControls::Construct(const FArguments& InArgs)
 			.OnClicked_Static(&SSourceControlControls::OnSourceControlCheckInChangesClicked)
 		]
 		+ SHorizontalBox::Slot() // Restore as Latest button
-		.VAlign(VAlign_Center)
 		.Padding(FMargin(0.0f, 0.0f, 4.0f, 0.0f))
 		.AutoWidth()
 		[
 			SNew(SButton)
-			.ButtonStyle(&FAppStyle::Get().GetWidgetStyle<FButtonStyle>("StatusBar.StatusBarButton"))
+			.ButtonStyle(FAppStyle::Get(), "SimpleButton")
 			.ToolTipText_Static(&SSourceControlControls::GetSourceControlRestoreAsLatestToolTipText)
 			.Visibility_Static(&SSourceControlControls::GetSourceControlRestoreAsLatestVisibility)
 			.IsEnabled_Static(&SSourceControlControls::IsSourceControlRestoreAsLatestEnabled)
@@ -150,7 +148,6 @@ void SSourceControlControls::Construct(const FArguments& InArgs)
 			.OnClicked_Static(&SSourceControlControls::OnSourceControlRestoreAsLatestClicked)
 		]
 		+SHorizontalBox::Slot() // Check In Kebab Combo button
-		.VAlign(VAlign_Center)
 		.AutoWidth()
 		[
 			SNew(SComboButton)
@@ -161,7 +158,6 @@ void SSourceControlControls::Construct(const FArguments& InArgs)
 			.OnGetMenuContent(InArgs._OnGenerateKebabMenu)
 		]
 		+ SHorizontalBox::Slot()
-		.VAlign(VAlign_Center)
 		.AutoWidth()
 		[
 			SNew(SSeparator)
@@ -170,11 +166,10 @@ void SSourceControlControls::Construct(const FArguments& InArgs)
 			.Orientation(EOrientation::Orient_Vertical)
 		]
 		+ SHorizontalBox::Slot() // Sync Latest Button
-		.VAlign(VAlign_Center)
 		.AutoWidth()
 		[
 			SNew(SButton)
-			.ButtonStyle(&FAppStyle::Get().GetWidgetStyle<FButtonStyle>("StatusBar.StatusBarButton"))
+			.ButtonStyle(FAppStyle::Get(), "SimpleButton")
 			.ToolTipText_Static(&SSourceControlControls::GetSourceControlSyncStatusToolTipText)
 			.Visibility_Static(&SSourceControlControls::GetSourceControlSyncStatusVisibility)
 			.IsEnabled_Static(&SSourceControlControls::IsSourceControlSyncEnabled)
@@ -201,7 +196,6 @@ void SSourceControlControls::Construct(const FArguments& InArgs)
 			.OnClicked_Static(&SSourceControlControls::OnSourceControlSyncClicked)
 		]
 		+ SHorizontalBox::Slot()
-		.VAlign(VAlign_Center)
 		.AutoWidth()
 		[
 			SNew(SSeparator)
