SourceControl: Assorted text changes (snapshot -> revision).

#rnx

#Synced-CL 49178553

diff --git a/Engine/Source/Developer/SourceControl/Private/SSourceControlControls.cpp b/Engine/Source/Developer/SourceControl/Private/SSourceControlControls.cpp
--- a/Engine/Source/Developer/SourceControl/Private/SSourceControlControls.cpp
+++ b/Engine/Source/Developer/SourceControl/Private/SSourceControlControls.cpp
@@ -299,14 +299,14 @@ FText SSourceControlControls::GetSourceControlSyncStatusToolTipText()
 	}
 	if (GetNumConflictsUpcoming() > 0)
 	{
-		return LOCTEXT("SyncLatestButtonNotAtHeadTooltipTextConflictUpcoming", "Some of your local changes conflict with later snapshots on the project. Sync Latest or Check-in Changes to review and address these issues.");
+		return LOCTEXT("SyncLatestButtonNotAtHeadTooltipTextConflictUpcoming", "Some of your local changes conflict with later revisions on the project. Sync Latest or Check-in Changes to review and address these issues.");
 	}
 	if (HasSourceControlChangesToSync())
 	{
-		return LOCTEXT("SyncLatestButtonNotAtHeadTooltipText", "Sync to the latest Snapshot for this project");
+		return LOCTEXT("SyncLatestButtonNotAtHeadTooltipText", "Sync to the latest revision for this project.");
 	}
 
-	return LOCTEXT("SyncLatestButtonAtHeadTooltipText", "Currently at the latest Snapshot for this project");
+	return LOCTEXT("SyncLatestButtonAtHeadTooltipText", "Currently at the latest revision for this project.");
 }
 
 const FSlateBrush* SSourceControlControls::GetSourceControlSyncStatusIcon()
@@ -439,14 +439,14 @@ FText SSourceControlControls::GetSourceControlCheckInStatusToolTipText()
 	}
 	if (GetNumConflictsUpcoming() > 0)
 	{
-		return LOCTEXT("CheckInButtonChangesTooltipTextConflictUpcoming", "Some of your local changes conflict with later snapshots on the project. Sync Latest or Check-in Changes to review and address these issues.");
+		return LOCTEXT("CheckInButtonChangesTooltipTextConflictUpcoming", "Some of your local changes conflict with later revisions on the project. Sync Latest or Check-in Changes to review and address these issues.");
 	}
 	if (HasSourceControlChangesToCheckIn())
 	{
-		return FText::Format(LOCTEXT("CheckInButtonChangesTooltipText", "Check-in {0} change(s) to this project"), GetNumLocalChanges());
+		return FText::Format(LOCTEXT("CheckInButtonChangesTooltipText", "Check-in {0} change(s) to this project."), GetNumLocalChanges());
 	}
 
-	return LOCTEXT("CheckInButtonNoChangesTooltipText", "No Changes to check in for this project");
+	return LOCTEXT("CheckInButtonNoChangesTooltipText", "No changes to check in for this project.");
 }
 
 const FSlateBrush* SSourceControlControls::GetSourceControlCheckInStatusIcon()
@@ -533,7 +533,7 @@ FText SSourceControlControls::GetSourceControlRestoreAsLatestToolTipText()
 		return LOCTEXT("RestoreAsLatestTooltipTextConflict", "Some of your local changes are still in conflict. Click here to review these conflicts.");
 	}
 	
-	return LOCTEXT("RestoreAsLatestTooltipText", "Restore this snapshot to be the latest version of the project for all team members.");
+	return LOCTEXT("RestoreAsLatestTooltipText", "Restore this revision to be the latest version of the project for all team members.");
 }
 
 const FSlateBrush* SSourceControlControls::GetSourceControlRestoreAsLatestStatusIcon()
