Update Source Control Styles

#rnx
#Synced-CL 49542191

diff --git a/Engine/Source/Developer/SourceControl/Private/RevisionControlStyle/RevisionControlStyle.cpp b/Engine/Source/Developer/SourceControl/Private/RevisionControlStyle/RevisionControlStyle.cpp
--- a/Engine/Source/Developer/SourceControl/Private/RevisionControlStyle/RevisionControlStyle.cpp
+++ b/Engine/Source/Developer/SourceControl/Private/RevisionControlStyle/RevisionControlStyle.cpp
@@ -187,6 +187,7 @@ FDefaultRevisionControlStyle::FDefaultRevisionControlStyle() : FSlateStyleSet(St
 	Set("RevisionControl.Rewound", new CORE_IMAGE_BRUSH_SVG("Starship/SourceControl/RC_Rewound", CoreStyleConstants::Icon26x26));
 	Set("RevisionControl.CheckInAvailableRewound", new CORE_IMAGE_BRUSH_SVG("Starship/SourceControl/RC_CheckInAvailableRewound", CoreStyleConstants::Icon26x26));
 	Set("RevisionControl.ConflictedState", new CORE_IMAGE_BRUSH_SVG("Starship/SourceControl/RC_ConflictedState", CoreStyleConstants::Icon26x26));
+	Set("RevisionControl.BranchPoint", new CORE_IMAGE_BRUSH_SVG("Starship/SourceControl/RC_BranchCircle", CoreStyleConstants::Icon26x26));
 
 	Set("RevisionControl.File", new CORE_IMAGE_BRUSH_SVG("Starship/SourceControl/RC_File", CoreStyleConstants::Icon6x8));
 	Set("RevisionControl.DiskSize", new CORE_IMAGE_BRUSH_SVG("Starship/SourceControl/RC_DiskSize", CoreStyleConstants::Icon8x8));
