Update filter parameters to GetSubmittedChangelists operation

#rnx
#Synced-CL 49542191

diff --git a/Engine/Source/Developer/SourceControl/Public/SourceControlOperations.h b/Engine/Source/Developer/SourceControl/Public/SourceControlOperations.h
--- a/Engine/Source/Developer/SourceControl/Public/SourceControlOperations.h
+++ b/Engine/Source/Developer/SourceControl/Public/SourceControlOperations.h
@@ -811,6 +811,16 @@ public:
 	}
 
 public:
+	void SetBranch(const FString& InBranch)
+	{
+		Branch = InBranch;
+	}
+
+	const FString& GetBranch()
+	{
+		return Branch;
+	}
+
 	const TArray<FSourceControlChangelistRef>& GetSubmittedChangelists() const { return Changelists; }
 
 	void AddSubmittedChangelist(FSourceControlChangelistRef Changelist)
@@ -822,6 +832,7 @@ private:
 	FDateTime DateFrom = FDateTime::MinValue();
 	FDateTime DateTo = FDateTime::MinValue();
 	bool Owned = false;
+	FString Branch;
 
 	int32 PaginationLimit = -1;
 	int32 PaginationOffset = -1;
