Update source control icons

#rnx
#Synced-CL 49629238

diff --git a/Engine/Source/Developer/SourceControl/Private/RevisionControlStyle/RevisionControlStyle.cpp b/Engine/Source/Developer/SourceControl/Private/RevisionControlStyle/RevisionControlStyle.cpp
--- a/Engine/Source/Developer/SourceControl/Private/RevisionControlStyle/RevisionControlStyle.cpp
+++ b/Engine/Source/Developer/SourceControl/Private/RevisionControlStyle/RevisionControlStyle.cpp
@@ -181,13 +181,15 @@ FDefaultRevisionControlStyle::FDefaultRevisionControlStyle() : FSlateStyleSet(St
 	Set("RevisionControl.VerticalLine", new CORE_IMAGE_BRUSH_SVG("Starship/SourceControl/RC_VerticalLine", CoreStyleConstants::Icon26x26));
 	Set("RevisionControl.VerticalLineStart", new CORE_IMAGE_BRUSH_SVG("Starship/SourceControl/RC_VerticalLineStart", CoreStyleConstants::Icon26x26));
 	Set("RevisionControl.VerticalLineDashed", new CORE_IMAGE_BRUSH_SVG("Starship/SourceControl/RC_VerticalLineDashed", CoreStyleConstants::Icon26x26));
-	Set("RevisionControl.CheckCircleLine", new CORE_IMAGE_BRUSH_SVG("Starship/SourceControl/RC_CheckCircleLine", CoreStyleConstants::Icon26x26));
+	Set("RevisionControl.CheckCircleLine_Check", new CORE_IMAGE_BRUSH_SVG("Starship/SourceControl/RC_CheckCircleLine_Check", CoreStyleConstants::Icon26x26));
+	Set("RevisionControl.CheckCircleLine_Circle", new CORE_IMAGE_BRUSH_SVG("Starship/SourceControl/RC_CheckCircleLine_Circle", CoreStyleConstants::Icon26x26));
 	Set("RevisionControl.LineCircle", new CORE_IMAGE_BRUSH_SVG("Starship/SourceControl/RC_LineCircle", CoreStyleConstants::Icon26x26));
 	Set("RevisionControl.CheckInAvailable", new CORE_IMAGE_BRUSH_SVG("Starship/SourceControl/RC_CheckInAvailable", CoreStyleConstants::Icon26x26));
 	Set("RevisionControl.Rewound", new CORE_IMAGE_BRUSH_SVG("Starship/SourceControl/RC_Rewound", CoreStyleConstants::Icon26x26));
 	Set("RevisionControl.CheckInAvailableRewound", new CORE_IMAGE_BRUSH_SVG("Starship/SourceControl/RC_CheckInAvailableRewound", CoreStyleConstants::Icon26x26));
 	Set("RevisionControl.ConflictedState", new CORE_IMAGE_BRUSH_SVG("Starship/SourceControl/RC_ConflictedState", CoreStyleConstants::Icon26x26));
-	Set("RevisionControl.BranchPoint", new CORE_IMAGE_BRUSH_SVG("Starship/SourceControl/RC_BranchCircle", CoreStyleConstants::Icon26x26));
+	Set("RevisionControl.BranchPoint_Branch", new CORE_IMAGE_BRUSH_SVG("Starship/SourceControl/RC_BranchCircle_Branch", CoreStyleConstants::Icon26x26));
+	Set("RevisionControl.BranchPoint_Circle", new CORE_IMAGE_BRUSH_SVG("Starship/SourceControl/RC_BranchCircle_Circle", CoreStyleConstants::Icon26x26));
 
 	Set("RevisionControl.File", new CORE_IMAGE_BRUSH_SVG("Starship/SourceControl/RC_File", CoreStyleConstants::Icon6x8));
 	Set("RevisionControl.DiskSize", new CORE_IMAGE_BRUSH_SVG("Starship/SourceControl/RC_DiskSize", CoreStyleConstants::Icon8x8));
