[Skein] Move source control settings from config file to source control CVars

Keep default value to FALSE and toggle them whenever necessary

#Synced-CL 49808600

diff --git a/Engine/Source/Developer/SourceControl/Private/SSourceControlControls.cpp b/Engine/Source/Developer/SourceControl/Private/SSourceControlControls.cpp
--- a/Engine/Source/Developer/SourceControl/Private/SSourceControlControls.cpp
+++ b/Engine/Source/Developer/SourceControl/Private/SSourceControlControls.cpp
@@ -7,6 +7,7 @@
 #include "ISourceControlModule.h"
 #include "Misc/ConfigCacheIni.h"
 #include "RevisionControlStyle/RevisionControlStyle.h"
+#include "SourceControlCVars.h"
 #include "Widgets/Input/SButton.h"
 #include "Widgets/Input/SComboBox.h"
 #include "Widgets/Layout/SSeparator.h"
@@ -30,8 +31,7 @@ FOnClicked SSourceControlControls::OnRestoreAsLatestClicked;
 
 static bool DisplaySyncStatus()
 {
-	bool bDisplaySourceControlSyncStatus = false;
-	GConfig->GetBool(TEXT("SourceControlSettings"), TEXT("DisplaySourceControlSyncStatus"), bDisplaySourceControlSyncStatus, GEditorIni);
+	const bool bDisplaySourceControlSyncStatus = SourceControlCVars::CVarSourceControlDisplaySyncStatus.GetValueOnAnyThread();
 
 	if (bDisplaySourceControlSyncStatus)
 	{
@@ -49,8 +49,7 @@ static bool DisplaySyncStatus()
 
 static bool DisplayCheckInStatus()
 {
-	bool bDisplaySourceControlCheckInStatus = false;
-	GConfig->GetBool(TEXT("SourceControlSettings"), TEXT("DisplaySourceControlCheckInStatus"), bDisplaySourceControlCheckInStatus, GEditorIni);
+	const bool bDisplaySourceControlCheckInStatus = SourceControlCVars::CVarSourceControlDisplayCheckInStatus.GetValueOnAnyThread();
 
 	if (bDisplaySourceControlCheckInStatus)
 	{
diff --git a/Engine/Source/Developer/SourceControl/Private/SourceControlCVars.cpp b/Engine/Source/Developer/SourceControl/Private/SourceControlCVars.cpp
--- a/Engine/Source/Developer/SourceControl/Private/SourceControlCVars.cpp
+++ b/Engine/Source/Developer/SourceControl/Private/SourceControlCVars.cpp
@@ -33,4 +33,16 @@ namespace SourceControlCVars
 		TEXT("SourceControl.IntervalBatchStatusUpdate"),
 		1,
 		TEXT("Controls interval that allows QueueStatusUpdate requests to be batched with others."));
+
+	TAutoConsoleVariable<bool> CVarSourceControlDisplaySyncStatus(
+		TEXT("SourceControl.DisplaySyncStatus"),
+		false,
+		TEXT("Allows SourceControl to display the Sync status.")
+	);
+
+	TAutoConsoleVariable<bool> CVarSourceControlDisplayCheckInStatus(
+		TEXT("SourceControl.DisplayCheckInStatus"),
+		false,
+		TEXT("Allows SourceControl to display the CheckIn status.")
+	);
 }
\ No newline at end of file
diff --git a/Engine/Source/Developer/SourceControl/Private/SourceControlCVars.h b/Engine/Source/Developer/SourceControl/Private/SourceControlCVars.h
--- a/Engine/Source/Developer/SourceControl/Private/SourceControlCVars.h
+++ b/Engine/Source/Developer/SourceControl/Private/SourceControlCVars.h
@@ -13,4 +13,7 @@ namespace SourceControlCVars
 	extern TAutoConsoleVariable<bool> CVarSourceControlEnableLoginDialogModal;
 
 	extern TAutoConsoleVariable<int32> CVarSourceControlIntervalBatchStatusUpdate;
+
+	extern TAutoConsoleVariable<bool> CVarSourceControlDisplaySyncStatus;
+	extern TAutoConsoleVariable<bool> CVarSourceControlDisplayCheckInStatus;
 }
\ No newline at end of file
