Update source control icons

#rnx

#Synced-CL 50152864

diff --git a/Engine/Source/Developer/SourceControl/Private/RevisionControlStyle/RevisionControlStyle.cpp b/Engine/Source/Developer/SourceControl/Private/RevisionControlStyle/RevisionControlStyle.cpp
--- a/Engine/Source/Developer/SourceControl/Private/RevisionControlStyle/RevisionControlStyle.cpp
+++ b/Engine/Source/Developer/SourceControl/Private/RevisionControlStyle/RevisionControlStyle.cpp
@@ -177,6 +177,9 @@ FDefaultRevisionControlStyle::FDefaultRevisionControlStyle() : FSlateStyleSet(St
 	Set("RevisionControl.BranchExplorer.Merge", new CORE_IMAGE_BRUSH_SVG("Starship/SourceControl/RC_Merged", CoreStyleConstants::Icon16x16));
 	Set("RevisionControl.BranchExplorer.NewBranch", new CORE_IMAGE_BRUSH_SVG("Starship/SourceControl/RC_NewBranch", CoreStyleConstants::Icon16x16));
 
+	Set("RevisionControl.Cross", new CORE_IMAGE_BRUSH_SVG("Starship/Common/close", CoreStyleConstants::Icon16x16));
+	Set("RevisionControl.Inspect", new CORE_IMAGE_BRUSH_SVG("Starship/Common/visible", CoreStyleConstants::Icon16x16));
+
 	// Revision Control States
 	Set("RevisionControl.VerticalLine", new CORE_IMAGE_BRUSH_SVG("Starship/SourceControl/RC_VerticalLine", CoreStyleConstants::Icon26x26));
 	Set("RevisionControl.VerticalLineStart", new CORE_IMAGE_BRUSH_SVG("Starship/SourceControl/RC_VerticalLineStart", CoreStyleConstants::Icon26x26));
