SourceControl: Fixed crash on FSourceControlState::GetHistoryItem.

If a SourceControl implementation clears the revision history of a file at an unexpected moment, the AssetDataCache should not crash.

#rnx

#Synced-CL 50815533

#Synced-CL 50815533

diff --git a/Engine/Source/Developer/SourceControl/Private/SourceControlAssetDataCache.cpp b/Engine/Source/Developer/SourceControl/Private/SourceControlAssetDataCache.cpp
--- a/Engine/Source/Developer/SourceControl/Private/SourceControlAssetDataCache.cpp
+++ b/Engine/Source/Developer/SourceControl/Private/SourceControlAssetDataCache.cpp
@@ -128,13 +128,13 @@ FSourceControlAssetDataEntry* FSourceControlAssetDataCache::AddAssetInformationE
 	{
 		check(!AssetDataEntry->bInitialized);
 
-		if (InFileState->GetHistorySize() == 0)
+		if (InFileState->GetHistorySize() > 0)
 		{
-			FileHistoryToUpdate.Add(Filename);
+			AssetDataToFetch.Enqueue(InFileState);
 		}
 		else
 		{
-			AssetDataToFetch.Enqueue(InFileState);
+			FileHistoryToUpdate.Add(Filename);
 		}
 	}
 	else
@@ -273,24 +273,28 @@ void FSourceControlAssetDataCache::LaunchFetchAssetDataTasks()
 
 void FSourceControlAssetDataCache::FetchAssetData(FSourceControlStateRef InFileState)
 {
-	check(InFileState->GetHistorySize() > 0);
-	
-	FString TempFileName;
-	TSharedPtr<ISourceControlRevision, ESPMode::ThreadSafe> Revision = InFileState->GetHistoryItem(0);
-	FSourceControlAssetDataEntry* AssetInformation = AssetDataCache.Find(InFileState->GetFilename());
-
-	check(AssetInformation != nullptr);
-	check(!AssetInformation->bInitialized);
-	check(AssetInformation->AssetDataArrayPtr.IsValid());
-	check(Revision.IsValid());
-
-	const int64 MaxFetchSize = (1 << 20); // 1MB
-	const bool bShouldGetFile = (MaxFetchSize < 0 || MaxFetchSize > static_cast<int64>(Revision->GetFileSize()));
-
-	if (bShouldGetFile && Revision->Get(TempFileName, EConcurrency::Asynchronous))
+	if (FSourceControlAssetDataEntry* AssetDataEntry = AssetDataCache.Find(InFileState->GetFilename()))
 	{
-		FAssetDataToLoad ToLoad = { TempFileName, InFileState->GetFilename() };
-		AssetDataToLoad.Enqueue(MoveTemp(ToLoad));
+		check(!AssetDataEntry->bInitialized);
+		check(AssetDataEntry->AssetDataArrayPtr.IsValid());
+
+		if (InFileState->GetHistorySize() > 0)
+		{
+			if (TSharedPtr<ISourceControlRevision, ESPMode::ThreadSafe> Revision = InFileState->GetHistoryItem(0))
+			{
+				const int64 MaxFetchSize = (1 << 20); // 1MB
+				const bool bShouldGetFile = (MaxFetchSize < 0 || MaxFetchSize > static_cast<int64>(Revision->GetFileSize()));
+				if (bShouldGetFile)
+				{
+					FString TempFileName;
+					if (Revision->Get(TempFileName, EConcurrency::Asynchronous))
+					{
+						FAssetDataToLoad ToLoad = { TempFileName, InFileState->GetFilename() };
+						AssetDataToLoad.Enqueue(MoveTemp(ToLoad));
+					}
+				}
+			}
+		}
 	}
 }
 
