Update source control style

#rnx

#Synced-CL 50944820

diff --git a/Engine/Source/Developer/SourceControl/Private/RevisionControlStyle/RevisionControlStyle.cpp b/Engine/Source/Developer/SourceControl/Private/RevisionControlStyle/RevisionControlStyle.cpp
--- a/Engine/Source/Developer/SourceControl/Private/RevisionControlStyle/RevisionControlStyle.cpp
+++ b/Engine/Source/Developer/SourceControl/Private/RevisionControlStyle/RevisionControlStyle.cpp
@@ -183,6 +183,7 @@ FDefaultRevisionControlStyle::FDefaultRevisionControlStyle() : FSlateStyleSet(St
 	// Revision Control States
 	Set("RevisionControl.VerticalLine", new CORE_IMAGE_BRUSH_SVG("Starship/SourceControl/RC_VerticalLine", CoreStyleConstants::Icon26x26));
 	Set("RevisionControl.VerticalLineStart", new CORE_IMAGE_BRUSH_SVG("Starship/SourceControl/RC_VerticalLineStart", CoreStyleConstants::Icon26x26));
+	Set("RevisionControl.VerticalLineEnd", new CORE_IMAGE_BRUSH_SVG("Starship/SourceControl/RC_VerticalLineEnd", CoreStyleConstants::Icon26x26));
 	Set("RevisionControl.VerticalLineDashed", new CORE_IMAGE_BRUSH_SVG("Starship/SourceControl/RC_VerticalLineDashed", CoreStyleConstants::Icon26x26));
 	Set("RevisionControl.CheckCircleLine_Check", new CORE_IMAGE_BRUSH_SVG("Starship/SourceControl/RC_CheckCircleLine_Check", CoreStyleConstants::Icon26x26));
 	Set("RevisionControl.CheckCircleLine_Circle", new CORE_IMAGE_BRUSH_SVG("Starship/SourceControl/RC_CheckCircleLine_Circle", CoreStyleConstants::Icon26x26));
