[Revision Control] Revision control helpers get status functions can now use the cached file state.

This can be useful after operations that already update file state (i.e. saving files)

diff --git a/Engine/Source/Developer/SourceControl/Private/SourceControlHelpers.cpp b/Engine/Source/Developer/SourceControl/Private/SourceControlHelpers.cpp
--- a/Engine/Source/Developer/SourceControl/Private/SourceControlHelpers.cpp
+++ b/Engine/Source/Developer/SourceControl/Private/SourceControlHelpers.cpp
@@ -240,7 +240,7 @@ bool ConvertFilesToQualifiedPaths(const TArray<FString>& InFiles, TArray<FString
 	return true;
 }
 
-void RunQueryFileStates(TConstArrayView<FString> InFiles, EConcurrency::Type InConcurrencyType, bool bSilent, TFunction<void(FSourceControlState)> FileStateCallback)
+void RunQueryFileStates(TConstArrayView<FString> InFiles, EConcurrency::Type InConcurrencyType, bool bSilent, bool bUseSourceControlStateCache, TFunction<void(FSourceControlState)> FileStateCallback)
 {
 	// Ensure source control system is up and running
 	ISourceControlProvider* Provider = SourceControlHelpersInternal::VerifySourceControl(bSilent);
@@ -266,27 +266,39 @@ void RunQueryFileStates(TConstArrayView<FString> InFiles, EConcurrency::Type InC
 		SCFilesToUpdate.Emplace(MoveTemp(SCFile));
 	}
 
-	// Make sure we update the modified state of the files (Perforce requires this
-	// since can be a more expensive test).
-	TSharedRef<FUpdateStatus, ESPMode::ThreadSafe> UpdateStatusOperation = ISourceControlOperation::Create<FUpdateStatus>();
-	UpdateStatusOperation->SetUpdateModifiedState(true);
+	auto SetStatusAndRunCallback = [](const TArray<FString>& InFilesToUpdate, const ISourceControlProvider* Provider, const TFunction<void(FSourceControlState)>& InCallbackFn)
+	{
+		TArray<FSourceControlStateRef> SCStates;
+		const_cast<ISourceControlProvider*>(Provider)->GetState(InFilesToUpdate, SCStates, EStateCacheUsage::Use);
 
-	ISourceControlModule::Get().GetProvider().Execute(
-		UpdateStatusOperation,
-		SCFilesToUpdate,
-		InConcurrencyType,
-		FSourceControlOperationComplete::CreateLambda([FileStateCallback, Provider, SCFilesToUpdate, bSilent](const FSourceControlOperationRef& InOperation, ECommandResult::Type InResult)
-			{
-				TArray<FSourceControlStateRef> SCStates;
-				Provider->GetState(SCFilesToUpdate, SCStates, EStateCacheUsage::Use);
+		for (const FSourceControlStatePtr SCState : SCStates)
+		{
+			FSourceControlState State;
+			State.SetFromStatus(SCState);
+			InCallbackFn(State);
+		}
+	};
 
-				for (const FSourceControlStatePtr SCState : SCStates)
-				{
-					FSourceControlState State;
-					State.SetFromStatus(SCState);
-					FileStateCallback(State);
-				}
+	if (!bUseSourceControlStateCache)
+	{
+		// Make sure we update the modified state of the files (Perforce requires this
+		// since can be a more expensive test).
+		TSharedRef<FUpdateStatus, ESPMode::ThreadSafe> UpdateStatusOperation = ISourceControlOperation::Create<FUpdateStatus>();
+		UpdateStatusOperation->SetUpdateModifiedState(true);
+
+		ISourceControlModule::Get().GetProvider().Execute(
+			UpdateStatusOperation,
+			SCFilesToUpdate,
+			InConcurrencyType,
+			FSourceControlOperationComplete::CreateLambda([FileStateCallback, Provider, SCFilesToUpdate, SetStatusAndRunCallback](const FSourceControlOperationRef& InOperation, ECommandResult::Type InResult)
+			{
+				SetStatusAndRunCallback(SCFilesToUpdate, Provider, FileStateCallback);
 			}));
+	}
+	else
+	{
+		SetStatusAndRunCallback(SCFilesToUpdate, Provider, FileStateCallback);
+	}
 }
 
 }  // namespace SourceControlHelpersInternal
@@ -1644,11 +1656,11 @@ bool USourceControlHelpers::CopyFile(const FString& InSourcePath, const FString&
 	return Result == ECommandResult::Succeeded;
 }
 
-TArray<FSourceControlState> USourceControlHelpers::QueryFileStates(const TArray<FString> InFiles, bool bSilent)
+TArray<FSourceControlState> USourceControlHelpers::QueryFileStates(const TArray<FString> InFiles, bool bSilent, bool bUseSourceControlStateCache)
 {
 	TArray<FSourceControlState> States;
 
-	SourceControlHelpersInternal::RunQueryFileStates(InFiles, EConcurrency::Synchronous, bSilent, 
+	SourceControlHelpersInternal::RunQueryFileStates(InFiles, EConcurrency::Synchronous, bSilent, bUseSourceControlStateCache,
 		[&States](FSourceControlState FileStateOut)
 		{
 			States.Emplace(MoveTemp(FileStateOut));
@@ -1657,11 +1669,11 @@ TArray<FSourceControlState> USourceControlHelpers::QueryFileStates(const TArray<
 	return States;
 }
 
-FSourceControlState USourceControlHelpers::QueryFileState(const FString& InFile, bool bSilent)
+FSourceControlState USourceControlHelpers::QueryFileState(const FString& InFile, bool bSilent, bool bUseSourceControlStateCache)
 {
 	FSourceControlState State;
 	State.Filename = InFile;
-	SourceControlHelpersInternal::RunQueryFileStates(TArray<FString> {InFile}, EConcurrency::Synchronous, bSilent,
+	SourceControlHelpersInternal::RunQueryFileStates(TArray<FString> {InFile}, EConcurrency::Synchronous, bSilent, bUseSourceControlStateCache,
 		[&State](FSourceControlState FileStateOut)
 		{
 			State = MoveTemp(FileStateOut);
@@ -1672,7 +1684,8 @@ FSourceControlState USourceControlHelpers::QueryFileState(const FString& InFile,
 
 void USourceControlHelpers::AsyncQueryFileStates(FQueryFileStateDelegate FileStateCallback, TArray<FString> InFiles, bool bSilent)
 {
-	SourceControlHelpersInternal::RunQueryFileStates(InFiles, EConcurrency::Asynchronous, bSilent,
+	bool constexpr bUseCachedSourceControlState = false;
+	SourceControlHelpersInternal::RunQueryFileStates(InFiles, EConcurrency::Asynchronous, bSilent, bUseCachedSourceControlState,
 		[FileStateCallback](FSourceControlState InState)
 		{
 			FileStateCallback.ExecuteIfBound(InState);
diff --git a/Engine/Source/Developer/SourceControl/Public/SourceControlHelpers.h b/Engine/Source/Developer/SourceControl/Public/SourceControlHelpers.h
--- a/Engine/Source/Developer/SourceControl/Public/SourceControlHelpers.h
+++ b/Engine/Source/Developer/SourceControl/Public/SourceControlHelpers.h
@@ -548,10 +548,10 @@ public:
 	 *			it could not have its values set.
 	 */
 	UFUNCTION(BlueprintCallable, Category = "Editor Scripting | Editor Revision Control Helpers", meta = (Keywords = "Source Control"))
-	static UE_API FSourceControlState QueryFileState(const FString& InFile, bool bSilent = false);
+	static UE_API FSourceControlState QueryFileState(const FString& InFile, bool bSilent = false, bool bUseSourceControlStateCache = false);
 
 	UFUNCTION(BlueprintCallable, Category = "Editor Scripting | Editor Revision Control Helpers", meta = (Keywords = "Source Control"))
-	static UE_API TArray<FSourceControlState> QueryFileStates(TArray<FString> InFiles, bool bSilent = false);
+	static UE_API TArray<FSourceControlState> QueryFileStates(TArray<FString> InFiles, bool bSilent = false, bool bUseSourceControlStateCache = false);
 
 	//Delegate to broadcast FileState upon AsyncQueryFileState completion
 	DECLARE_DYNAMIC_DELEGATE_OneParam(FQueryFileStateDelegate, FSourceControlState, FileStateOut);
