SourceControl: Change 'Check-in Changes' to 'Check In Changes'.

#rnx

#Synced-CL 51090019

diff --git a/Engine/Source/Developer/SourceControl/Private/SSourceControlControls.cpp b/Engine/Source/Developer/SourceControl/Private/SSourceControlControls.cpp
--- a/Engine/Source/Developer/SourceControl/Private/SSourceControlControls.cpp
+++ b/Engine/Source/Developer/SourceControl/Private/SSourceControlControls.cpp
@@ -298,7 +298,7 @@ FText SSourceControlControls::GetSourceControlSyncStatusToolTipText()
 	}
 	if (GetNumConflictsUpcoming() > 0)
 	{
-		return LOCTEXT("SyncLatestButtonNotAtHeadTooltipTextConflictUpcoming", "Some of your local changes conflict with later revisions on the project. Sync Latest or Check-in Changes to review and address these issues.");
+		return LOCTEXT("SyncLatestButtonNotAtHeadTooltipTextConflictUpcoming", "Some of your local changes conflict with later revisions on the project. Sync Latest or Check In Changes to review and address these issues.");
 	}
 	if (HasSourceControlChangesToSync())
 	{
@@ -424,7 +424,7 @@ FText SSourceControlControls::GetSourceControlCheckInStatusText()
 {
 	if (HasSourceControlChangesToCheckIn())
 	{
-		return LOCTEXT("CheckInButtonChangesText", "Check-in Changes");
+		return LOCTEXT("CheckInButtonChangesText", "Check In Changes");
 	}
 
 	return LOCTEXT("CheckInButtonNoChangesText", "No Changes");
@@ -438,7 +438,7 @@ FText SSourceControlControls::GetSourceControlCheckInStatusToolTipText()
 	}
 	if (GetNumConflictsUpcoming() > 0)
 	{
-		return LOCTEXT("CheckInButtonChangesTooltipTextConflictUpcoming", "Some of your local changes conflict with later revisions on the project. Sync Latest or Check-in Changes to review and address these issues.");
+		return LOCTEXT("CheckInButtonChangesTooltipTextConflictUpcoming", "Some of your local changes conflict with later revisions on the project. Sync Latest or Check In Changes to review and address these issues.");
 	}
 	if (HasSourceControlChangesToCheckIn())
 	{
diff --git a/Engine/Source/Developer/SourceControlCheckInPrompt/Private/SourceControlCheckInPromptModule.cpp b/Engine/Source/Developer/SourceControlCheckInPrompt/Private/SourceControlCheckInPromptModule.cpp
--- a/Engine/Source/Developer/SourceControlCheckInPrompt/Private/SourceControlCheckInPromptModule.cpp
+++ b/Engine/Source/Developer/SourceControlCheckInPrompt/Private/SourceControlCheckInPromptModule.cpp
@@ -50,11 +50,11 @@ void FSourceControlCheckInPromptModule::ShutdownModule()
 
 void FSourceControlCheckInPromptModule::ShowModal(const FText& InMessage)
 {
-	FText Title = LOCTEXT("CheckInModal_Title", "Check-in Changes");
+	FText Title = LOCTEXT("CheckInModal_Title", "Check In Changes");
 
 	FSuppressableWarningDialog::FSetupInfo Info(InMessage, Title, "CheckInModal_Warning");
-	Info.ConfirmText = LOCTEXT("CheckInModal_Yes", "Check-in Changes");
-	Info.CancelText = LOCTEXT("CheckInModal_No", "Skip Check-in");
+	Info.ConfirmText = LOCTEXT("CheckInModal_Yes", "Check In Changes");
+	Info.CancelText = LOCTEXT("CheckInModal_No", "Skip Check-In");
 	Info.CheckBoxText = FText::GetEmpty(); // Not suppressable!
 	Info.Image = const_cast<FSlateBrush*>(FRevisionControlStyleManager::Get().GetBrush("RevisionControl.StatusBar.HasLocalChanges"));
 
diff --git a/Engine/Source/Editor/StatusBar/Private/SourceControlMenuHelpers.cpp b/Engine/Source/Editor/StatusBar/Private/SourceControlMenuHelpers.cpp
--- a/Engine/Source/Editor/StatusBar/Private/SourceControlMenuHelpers.cpp
+++ b/Engine/Source/Editor/StatusBar/Private/SourceControlMenuHelpers.cpp
@@ -244,7 +244,7 @@ void FSourceControlCommands::CheckOutModifiedFiles_Clicked()
 
 bool FSourceControlCommands::RevertAllModifiedFiles_CanExecute()
 {
-	// If CHECK-IN CHANGES button is active, the REVERT ALL option should be active.
+	// If CHECK IN CHANGES button is active, the REVERT ALL option should be active.
 
 	if (FUnsavedAssetsTrackerModule::Get().GetUnsavedAssetNum() > 0)
 	{
