SourceControl: Change verb 'Check-in' to 'Check in'.

#rnx

#Synced-CL 51090019

diff --git a/Engine/Source/Developer/SourceControl/Private/SSourceControlControls.cpp b/Engine/Source/Developer/SourceControl/Private/SSourceControlControls.cpp
--- a/Engine/Source/Developer/SourceControl/Private/SSourceControlControls.cpp
+++ b/Engine/Source/Developer/SourceControl/Private/SSourceControlControls.cpp
@@ -442,7 +442,7 @@ FText SSourceControlControls::GetSourceControlCheckInStatusToolTipText()
 	}
 	if (HasSourceControlChangesToCheckIn())
 	{
-		return FText::Format(LOCTEXT("CheckInButtonChangesTooltipText", "Check-in {0} change(s) to this project."), GetNumLocalChanges());
+		return FText::Format(LOCTEXT("CheckInButtonChangesTooltipText", "Check in {0} change(s) to this project."), GetNumLocalChanges());
 	}
 
 	return LOCTEXT("CheckInButtonNoChangesTooltipText", "No changes to check in for this project.");
