Update source control style

#rnx
#Synced-CL 51536473

diff --git a/Engine/Source/Developer/SourceControl/Private/RevisionControlStyle/RevisionControlStyle.cpp b/Engine/Source/Developer/SourceControl/Private/RevisionControlStyle/RevisionControlStyle.cpp
--- a/Engine/Source/Developer/SourceControl/Private/RevisionControlStyle/RevisionControlStyle.cpp
+++ b/Engine/Source/Developer/SourceControl/Private/RevisionControlStyle/RevisionControlStyle.cpp
@@ -180,6 +180,8 @@ FDefaultRevisionControlStyle::FDefaultRevisionControlStyle() : FSlateStyleSet(St
 	Set("RevisionControl.Cross", new CORE_IMAGE_BRUSH_SVG("Starship/Common/close", CoreStyleConstants::Icon16x16));
 	Set("RevisionControl.Inspect", new CORE_IMAGE_BRUSH_SVG("Starship/Common/visible", CoreStyleConstants::Icon16x16));
 
+	Set("RevisionControl.Filter.ResetToDefault", new CORE_IMAGE_BRUSH_SVG("Starship/Common/ResetToDefault", CoreStyleConstants::Icon16x16));
+
 	// Revision Control States
 	Set("RevisionControl.VerticalLine", new CORE_IMAGE_BRUSH_SVG("Starship/SourceControl/RC_VerticalLine", CoreStyleConstants::Icon26x26));
 	Set("RevisionControl.VerticalLineStart", new CORE_IMAGE_BRUSH_SVG("Starship/SourceControl/RC_VerticalLineStart", CoreStyleConstants::Icon26x26));
