SourceControl: Add logging of timing statistics to ApplyOperationAndReloadPackages.

#rnx
#Synced-CL 52360262

diff --git a/Engine/Source/Developer/SourceControl/Private/SourceControlHelpers.cpp b/Engine/Source/Developer/SourceControl/Private/SourceControlHelpers.cpp
--- a/Engine/Source/Developer/SourceControl/Private/SourceControlHelpers.cpp
+++ b/Engine/Source/Developer/SourceControl/Private/SourceControlHelpers.cpp
@@ -25,10 +25,12 @@
 
 #if WITH_EDITOR
 #include "Editor.h"
+#include "EngineAnalytics.h"
 #include "PackageTools.h"
 #include "ObjectTools.h"
 #include "FileHelpers.h"
 #include "Bookmarks/BookmarkScoped.h"
+#include "Interfaces/IAnalyticsProvider.h"
 #endif
 
 #define LOCTEXT_NAMESPACE "SourceControlHelpers"
@@ -975,7 +977,12 @@ bool USourceControlHelpers::RevertFile(const FString& InFile, bool bSilent)
 #if WITH_EDITOR
 bool USourceControlHelpers::ApplyOperationAndReloadPackages(const TArray<FString>& InFilenames,	const TFunctionRef<bool(const TArray<FString>&)>& InOperation, bool bReloadWorld, bool bInteractive)
 {
-	ISourceControlProvider& SourceControlProvider = ISourceControlModule::Get().GetProvider();
+	const double TimeStarted = FPlatformTime::Seconds();
+
+	double TimeTakenFlush = 0;
+	double TimeTakenReload = 0;
+	double TimeTakenOperation = 0;
+	double TimeTakenUnload = 0;
 
 	TArray<FString> PackageNames;
 	TArray<FString> PackageFilenames;
@@ -1071,7 +1078,7 @@ bool USourceControlHelpers::ApplyOperationAndReloadPackages(const TArray<FString
 	{
 		TStringBuilder<2048> Builder;
 		Builder.Join(FilteredPackages, TEXT(", "));
-		const FString Packages = Builder.ToString();
+		const FString Packages = *Builder;
 
 		UE_LOG(LogSourceControl, Warning, TEXT("This operation could not complete on the following map or external packages, please unload them before retrying : %s"), *Packages);
 		return false;
@@ -1144,6 +1151,8 @@ bool USourceControlHelpers::ApplyOperationAndReloadPackages(const TArray<FString
 	UnlinkPackages.Append(LoadedPackages);
 	if (UnlinkPackages.Num() > 0)
 	{
+		const double TimeStartedFlush = FPlatformTime::Seconds();
+
 		TArray<UObject*> LoadedObjects;
 		TArray<int32> PendingPackageRequestIds;
 		LoadedObjects.Reserve(UnlinkPackages.Num());
@@ -1160,6 +1169,8 @@ bool USourceControlHelpers::ApplyOperationAndReloadPackages(const TArray<FString
 
 		FlushAsyncLoading(PendingPackageRequestIds);
 		ResetLoaders(LoadedObjects);
+
+		TimeTakenFlush = FPlatformTime::Seconds() - TimeStartedFlush;
 	}
 
 	for (int32 PackageIndex = 0; PackageIndex < PackageNames.Num(); PackageIndex++)
@@ -1175,7 +1186,13 @@ bool USourceControlHelpers::ApplyOperationAndReloadPackages(const TArray<FString
 	}
 
 	// Apply Operation
-	bSuccess = InOperation(PackageFilenames);
+	{
+		const double TimeStartedOperation = FPlatformTime::Seconds();
+
+		bSuccess = InOperation(PackageFilenames);
+
+		TimeTakenOperation = FPlatformTime::Seconds() - TimeStartedOperation;
+	}
 
 	// Reverting may have deleted some packages, so we need to delete those and unload them rather than re-load them...
 	FAssetRegistryModule& AssetRegistryModule = FModuleManager::LoadModuleChecked<FAssetRegistryModule>(TEXT("AssetRegistry"));
@@ -1204,6 +1221,8 @@ bool USourceControlHelpers::ApplyOperationAndReloadPackages(const TArray<FString
 	// Hot-reload the new packages...
 	if (LoadedPackages.Num() > 0)
 	{
+		const double TimeStartedReload = FPlatformTime::Seconds();
+
 		FBookmarkScoped BookmarkScoped;
 		FScopedSlowTask SlowTaskScoped(1.f, NSLOCTEXT("SourceControl", "Reload", "Reloading..."));
 		SlowTaskScoped.MakeDialog();
@@ -1249,6 +1268,8 @@ bool USourceControlHelpers::ApplyOperationAndReloadPackages(const TArray<FString
 				UE_LOG(LogSourceControl, Warning, TEXT("%s"), *OutReloadErrorMsg.ToString());
 			}
 		}
+
+		TimeTakenReload = FPlatformTime::Seconds() - TimeStartedReload;
 	}
 
 	// A world reload might have already deleted some objects, so check which missing ones are still valid.
@@ -1267,14 +1288,60 @@ bool USourceControlHelpers::ApplyOperationAndReloadPackages(const TArray<FString
 	// Delete and Unload assets...
 	if (ObjectsToDelete.Num() > 0)
 	{
+		const double TimeStartedUnload = FPlatformTime::Seconds();
+
 		if (ObjectTools::DeleteObjectsUnchecked(ObjectsToDelete) != ObjectsToDelete.Num())
 		{
 			UE_LOG(LogSourceControl, Warning, TEXT("Failed to unload some assets."));
 		}
+
+		TimeTakenUnload = FPlatformTime::Seconds() - TimeStartedUnload;
 	}
-	
+
+
 	// Re-cache the SCC state...
-	SourceControlProvider.Execute(ISourceControlOperation::Create<FUpdateStatus>(), PackageFilenames, EConcurrency::Asynchronous);
+	ISourceControlModule::Get().GetProvider().Execute(ISourceControlOperation::Create<FUpdateStatus>(), PackageFilenames, EConcurrency::Asynchronous);
+
+	// Log timings...
+	const double TimeTaken = FPlatformTime::Seconds() - TimeStarted;
+	if (TimeTaken > 1.f)
+	{
+		TStringBuilder<512> TimingMessage;
+		TimingMessage.Appendf(TEXT("ApplyOperationAndReloadPackages on %d file(s) took %.3fs"), InFilenames.Num(), TimeTaken);
+		if (TimeTakenOperation > 0)
+		{
+			TimingMessage.Appendf(TEXT(" - operation: %.3fs"), TimeTakenOperation);
+		}
+		if (TimeTakenReload > 0)
+		{
+			TimingMessage.Appendf(TEXT(" - reload: %.3fs"), TimeTakenReload);
+		}
+		if (TimeTakenFlush > 0)
+		{
+			TimingMessage.Appendf(TEXT(" - flush: %.3fs"), TimeTakenFlush);
+		}
+		if (TimeTakenUnload > 0)
+		{
+			TimingMessage.Appendf(TEXT(" - unload: %.3fs"), TimeTakenUnload);
+		}
+		if (bReloadWorld)
+		{
+			TimingMessage.Appendf(TEXT(" - ReloadWorld"));
+		}
+		UE_LOGF(LogSourceControl, Log, "%ls", *TimingMessage);
+
+		if (FEngineAnalytics::IsAvailable())
+		{
+			TArray<FAnalyticsEventAttribute> Attributes;
+			Attributes.Emplace(TEXT("Duration.Seconds"), TimeTaken);
+			Attributes.Emplace(TEXT("Duration.Seconds.Operation"), TimeTakenOperation);
+			Attributes.Emplace(TEXT("Duration.Seconds.Reload"), TimeTakenReload);
+			Attributes.Emplace(TEXT("Duration.Seconds.Flush"), TimeTakenFlush);
+			Attributes.Emplace(TEXT("Duration.Seconds.Unload"), TimeTakenUnload);
+
+			FEngineAnalytics::GetProvider().RecordEvent(TEXT("Editor.Usage.SourceControl.ApplyOperationAndReloadPackages"), Attributes);
+		}
+	}
 
 	return bSuccess;
 }
